Input Widgets: Advanced Rotation and Transform widget

* Rotation to support multiple representations (quat, rotator, euler etc)
* Rotation widget to support mapping to single numeric channels or higher level attributes (Rotator, Quat)
* Transform widget relative / world
* Transform widget uniform scale lock
* Transform widget to map to single numeric channels or higher level attributes (FTransform, FEulerTransform)

diff --git a/Engine/Source/Developer/AnimationWidgets/AnimationWidgets.Build.cs b/Engine/Source/Developer/AnimationWidgets/AnimationWidgets.Build.cs
new file mode 100644
--- /dev/null
+++ b/Engine/Source/Developer/AnimationWidgets/AnimationWidgets.Build.cs
@@ -0,0 +1,34 @@
+// Copyright Epic Games, Inc. All Rights Reserved.
+
+using UnrealBuildTool;
+
+public class AnimationWidgets : ModuleRules
+{
+	public AnimationWidgets(ReadOnlyTargetRules Target) : base(Target)
+	{
+		/** NOTE: THIS MODULE SHOULD NOT EVER DEPEND ON UNREALED. 
+		 * Please refer to ToolWidgets.Build.cs for more information.
+		 */
+		PublicDependencyModuleNames.AddRange(
+			new string[] {
+				"Core",
+				"Slate",
+				"SlateCore",
+				"InputCore",
+				"ToolWidgets",
+				"AnimationCore",
+			}
+		);
+
+		PrivateDependencyModuleNames.AddRange(
+			new string[] {
+			}
+		);
+
+		PrivateIncludePaths.AddRange(
+			new string[] {
+				"Developer/AnimationWidgets/Private",
+			}
+		);
+	}
+}
diff --git a/Engine/Source/Developer/AnimationWidgets/Private/AnimationWidgetsModule.cpp b/Engine/Source/Developer/AnimationWidgets/Private/AnimationWidgetsModule.cpp
new file mode 100644
--- /dev/null
+++ b/Engine/Source/Developer/AnimationWidgets/Private/AnimationWidgetsModule.cpp
@@ -0,0 +1,18 @@
+// Copyright Epic Games, Inc. All Rights Reserved.
+
+#include "Modules/ModuleInterface.h"
+#include "Modules/ModuleManager.h"
+
+class FToolWidgetsModule
+	: public IModuleInterface
+{
+public:
+
+	// IModuleInterface interface
+	
+	virtual void StartupModule() override { }
+	virtual void ShutdownModule() override { }
+};
+
+
+IMPLEMENT_MODULE(FToolWidgetsModule, ToolWidgets);
diff --git a/Engine/Source/Developer/AnimationWidgets/Public/SAdvancedRotationInputBox.h b/Engine/Source/Developer/AnimationWidgets/Public/SAdvancedRotationInputBox.h
new file mode 100644
--- /dev/null
+++ b/Engine/Source/Developer/AnimationWidgets/Public/SAdvancedRotationInputBox.h
@@ -0,0 +1,905 @@
+// Copyright Epic Games, Inc. All Rights Reserved.
+
+#pragma once
+
+#include "CoreMinimal.h"
+#include "AnimationCoreLibrary.h"
+#include "Layout/Margin.h"
+#include "Fonts/SlateFontInfo.h"
+#include "Widgets/DeclarativeSyntaxSupport.h"
+#include "Widgets/SWidget.h"
+#include "Widgets/SCompoundWidget.h"
+#include "Widgets/Input/SRotatorInputBox.h"
+#include "Widgets/Input/SVectorInputBox.h"
+#include "Styling/CoreStyle.h"
+#include "Framework/SlateDelegates.h"
+#include "Widgets/Input/NumericTypeInterface.h"
+#include "Widgets/Input/NumericUnitTypeInterface.inl"
+
+namespace ESlateRotationRepresentation
+{
+	enum Type
+	{
+		EulerXYZ,
+		EulerXZY,
+		EulerYXZ,
+		EulerYZX,
+		EulerZXY,
+		EulerZYX,
+		Rotator,
+		Quaternion,
+		AxisAndAngle,
+		Max
+	};
+}
+
+namespace ESlateTransformSubComponent
+{
+	enum Type
+	{
+		X,
+		Y,
+		Z,
+		W,
+		Angle,
+		Pitch,
+		Yaw,
+		Roll,
+		Max
+	};
+}
+
+/**
+ * Generic Rotation Slate control
+ */
+template<typename NumericType = FVector::FReal>
+class SAdvancedRotationInputBox : public SCompoundWidget
+{
+public:
+	
+	typedef SNumericVectorInputBox<NumericType, UE::Math::TVector<NumericType>, 3> SNumericVectorInputBox3; 
+	typedef SNumericVectorInputBox<NumericType, UE::Math::TVector<NumericType>, 4> SNumericVectorInputBox4; 
+
+	/** Delegate for notification for a rotator value change */
+ 	DECLARE_DELEGATE_OneParam(FOnRotatorChanged, FRotator);
+
+	/** Delegate for notification for a rotator value commit */
+	DECLARE_DELEGATE_TwoParams(FOnRotatorCommitted, FRotator, ETextCommit::Type);
+
+	/** Delegate for notification for a quaternion value change */
+	DECLARE_DELEGATE_OneParam(FOnQuaternionChanged, FQuat);
+
+	/** Delegate for notification for a quaternion value commit */
+	DECLARE_DELEGATE_TwoParams(FOnQuaternionCommitted, FQuat, ETextCommit::Type);
+
+	/** Delegate to retrieve a numerictype value  */
+	DECLARE_DELEGATE_RetVal_TwoParams(TOptional<NumericType>, FOnGetNumericValue, ESlateRotationRepresentation::Type, ESlateTransformSubComponent::Type);
+
+	/** Notification for numerictype value change */
+ 	DECLARE_DELEGATE_ThreeParams(FOnNumericValueChanged, ESlateRotationRepresentation::Type, ESlateTransformSubComponent::Type, NumericType);
+
+	/** Notification for numerictype value committed */
+	DECLARE_DELEGATE_FourParams(FOnNumericValueCommitted, ESlateRotationRepresentation::Type, ESlateTransformSubComponent::Type, NumericType, ETextCommit::Type);
+
+	/** Delegate to retrieve toggle checkstate */
+	DECLARE_DELEGATE_RetVal_TwoParams(ECheckBoxState, FOnGetToggleChecked, ESlateRotationRepresentation::Type, ESlateTransformSubComponent::Type);
+
+	/** Notification for toggle checkstate change */
+	DECLARE_DELEGATE_ThreeParams(FOnToggleChanged, ESlateRotationRepresentation::Type, ESlateTransformSubComponent::Type, ECheckBoxState);
+
+	SLATE_BEGIN_ARGS( SAdvancedRotationInputBox )
+		: _Representation(ESlateRotationRepresentation::Rotator)
+		, _bColorAxisLabels(true)
+		, _Font(FAppStyle::Get().GetFontStyle("NormalFont"))
+		, _AllowSpin(true)
+		, _DisplayToggle( false )
+		, _TogglePadding( FMargin( 1.f,0.f,1.f,0.f ) )
+		{}
+
+		/** Representation of the rotation */
+		SLATE_ATTRIBUTE( ESlateRotationRepresentation::Type, Representation )
+
+		/** Optional representation as rotator */
+		SLATE_ATTRIBUTE( TOptional<FRotator>, Rotator )
+
+		/** Optional delegate to notify a change in the rotator */
+		SLATE_EVENT( FOnRotatorChanged, OnRotatorChanged )
+
+		/** Optional delegate to notify a commit in the rotator */
+		SLATE_EVENT( FOnRotatorCommitted, OnRotatorCommitted )
+
+		/** Optional representation as quaternion */
+		SLATE_ATTRIBUTE( TOptional<FQuat>, Quaternion )
+
+		/** Optional delegate to notify a change in the quaternion */
+		SLATE_EVENT( FOnQuaternionChanged, OnQuaternionChanged )
+
+		/** Optional delegate to notify a commit in the quaternion */
+		SLATE_EVENT( FOnQuaternionCommitted, OnQuaternionCommitted )
+
+		/** Optional delegate to retrieve a value */
+		SLATE_EVENT( FOnGetNumericValue, OnGetNumericValue )
+
+		/** Optional delegate to notify a value change */
+		SLATE_EVENT( FOnNumericValueChanged, OnNumericValueChanged )
+
+		/** Optional delegate to notify a committed value */
+		SLATE_EVENT( FOnNumericValueCommitted, OnNumericValueCommitted )
+
+		/** Should the axis labels be colored */
+		SLATE_ARGUMENT( bool, bColorAxisLabels )		
+
+		/** Font to use for the text in this box */
+		SLATE_ATTRIBUTE( FSlateFontInfo, Font )
+
+		/** Whether or not values can be spun or if they should be typed in */
+		SLATE_ARGUMENT( bool, AllowSpin )
+
+		/** Whether or not to include a toggle checkbox to the left of the widget */
+		SLATE_ARGUMENT( bool, DisplayToggle )
+
+		/** Delegate to retrieve a toggle state */
+		SLATE_EVENT( FOnGetToggleChecked, OnGetToggleChecked )
+
+		/** Delegate to notify a toggle state change */
+		SLATE_EVENT( FOnToggleChanged, OnToggleChanged )
+
+		/** Padding around the toggle checkbox */
+		SLATE_ARGUMENT( FMargin, TogglePadding )
+
+	SLATE_END_ARGS()
+
+	/**
+	 * Construct this widget
+	 *
+	 * @param	InArgs	The declaration data for this widget
+	 */
+	void Construct( const FArguments& InArgs )
+	{
+		Representation = InArgs._Representation;
+		Rotator = InArgs._Rotator;
+		OnRotatorChanged = InArgs._OnRotatorChanged;
+		OnRotatorCommitted = InArgs._OnRotatorCommitted;
+		Quaternion = InArgs._Quaternion;
+		OnQuaternionChanged = InArgs._OnQuaternionChanged;
+		OnQuaternionCommitted = InArgs._OnQuaternionCommitted;
+		OnGetNumericValue = InArgs._OnGetNumericValue; 
+		OnNumericValueChanged = InArgs._OnNumericValueChanged; 
+		OnNumericValueCommitted = InArgs._OnNumericValueCommitted; 
+		OnGetToggleChecked = InArgs._OnGetToggleChecked; 
+		OnToggleChanged = InArgs._OnToggleChanged;
+
+		TSharedRef<SWidget> AngleLabelWidget = SNullWidget::NullWidget;
+		if (InArgs._bColorAxisLabels)
+		{
+			const FLinearColor AngleLabelColor = SNumericEntryBox<NumericType>::GreenLabelBackgroundColor;
+			AngleLabelWidget = SNumericEntryBox<NumericType>::BuildNarrowColorLabel(AngleLabelColor);
+		}
+
+		TSharedPtr<TNumericUnitTypeInterface<NumericType>> DegreesTypeInterface =
+			MakeShareable( new TNumericUnitTypeInterface<NumericType>(EUnit::Degrees) );
+
+		this->ChildSlot
+		[
+			SNew(SHorizontalBox)
+
+			// a rotator variation - which is the default in the editor
+			+SHorizontalBox::Slot()
+			.Padding(FMargin(0.f))
+			[
+				SNew(SNumericRotatorInputBox<NumericType>)
+				.bColorAxisLabels(InArgs._bColorAxisLabels)
+				.Font(InArgs._Font)
+				.AllowSpin(InArgs._AllowSpin)
+				.TypeInterface(DegreesTypeInterface)
+				.Roll(this, &SAdvancedRotationInputBox::HandleGetNumericValue, ESlateTransformSubComponent::Roll)
+				.Yaw(this, &SAdvancedRotationInputBox::HandleGetNumericValue, ESlateTransformSubComponent::Yaw)
+				.Pitch(this, &SAdvancedRotationInputBox::HandleGetNumericValue, ESlateTransformSubComponent::Pitch)
+				.OnRollChanged(this, &SAdvancedRotationInputBox::HandleOnNumericValueChanged, ESlateTransformSubComponent::Roll)
+				.OnYawChanged(this, &SAdvancedRotationInputBox::HandleOnNumericValueChanged, ESlateTransformSubComponent::Yaw)
+				.OnPitchChanged(this, &SAdvancedRotationInputBox::HandleOnNumericValueChanged, ESlateTransformSubComponent::Pitch)
+				.OnRollCommitted(this, &SAdvancedRotationInputBox::HandleOnNumericValueCommitted, ESlateTransformSubComponent::Roll)
+				.OnYawCommitted(this, &SAdvancedRotationInputBox::HandleOnNumericValueCommitted, ESlateTransformSubComponent::Yaw)
+				.OnPitchCommitted(this, &SAdvancedRotationInputBox::HandleOnNumericValueCommitted, ESlateTransformSubComponent::Pitch)
+				.DisplayToggle(InArgs._DisplayToggle)
+				.ToggleRollChecked(this, &SAdvancedRotationInputBox::HandleGetToggleChecked, ESlateTransformSubComponent::Roll)
+				.ToggleYawChecked(this, &SAdvancedRotationInputBox::HandleGetToggleChecked, ESlateTransformSubComponent::Yaw)
+				.TogglePitchChecked(this, &SAdvancedRotationInputBox::HandleGetToggleChecked, ESlateTransformSubComponent::Pitch)
+				.OnToggleRollChanged(this, &SAdvancedRotationInputBox::HandleOnToggleChanged, ESlateTransformSubComponent::Roll)
+				.OnToggleYawChanged(this, &SAdvancedRotationInputBox::HandleOnToggleChanged, ESlateTransformSubComponent::Yaw)
+				.OnTogglePitchChanged(this, &SAdvancedRotationInputBox::HandleOnToggleChanged, ESlateTransformSubComponent::Pitch)
+				.TogglePadding(InArgs._TogglePadding)
+				.Visibility(this, &SAdvancedRotationInputBox::IsRotatorInputBoxVisible)
+			]
+
+			// a vector 3 variation to represent euler with rotation order
+			+SHorizontalBox::Slot()
+			.Padding(FMargin(0.f))
+			[
+				SNew(SNumericVectorInputBox3)
+				.bColorAxisLabels(InArgs._bColorAxisLabels)
+				.Font(InArgs._Font)
+				.AllowSpin(InArgs._AllowSpin)
+				.SpinDelta(0.01f)
+				.MinSliderVector(FVector(-360.f, -360.f, -360.f))
+				.MaxSliderVector(FVector(360.f, 360.f, 360.f))
+				.TypeInterface(DegreesTypeInterface)
+				.X(this, &SAdvancedRotationInputBox::HandleGetNumericValue, ESlateTransformSubComponent::X)
+				.Y(this, &SAdvancedRotationInputBox::HandleGetNumericValue, ESlateTransformSubComponent::Y)
+				.Z(this, &SAdvancedRotationInputBox::HandleGetNumericValue, ESlateTransformSubComponent::Z)
+				.OnXChanged(this, &SAdvancedRotationInputBox::HandleOnNumericValueChanged, ESlateTransformSubComponent::X)
+				.OnYChanged(this, &SAdvancedRotationInputBox::HandleOnNumericValueChanged, ESlateTransformSubComponent::Y)
+				.OnZChanged(this, &SAdvancedRotationInputBox::HandleOnNumericValueChanged, ESlateTransformSubComponent::Z)
+				.OnXCommitted(this, &SAdvancedRotationInputBox::HandleOnNumericValueCommitted, ESlateTransformSubComponent::X)
+				.OnYCommitted(this, &SAdvancedRotationInputBox::HandleOnNumericValueCommitted, ESlateTransformSubComponent::Y)
+				.OnZCommitted(this, &SAdvancedRotationInputBox::HandleOnNumericValueCommitted, ESlateTransformSubComponent::Z)
+				.DisplayToggle(InArgs._DisplayToggle)
+				.ToggleXChecked(this, &SAdvancedRotationInputBox::HandleGetToggleChecked, ESlateTransformSubComponent::X)
+				.ToggleYChecked(this, &SAdvancedRotationInputBox::HandleGetToggleChecked, ESlateTransformSubComponent::Y)
+				.ToggleZChecked(this, &SAdvancedRotationInputBox::HandleGetToggleChecked, ESlateTransformSubComponent::Z)
+				.OnToggleXChanged(this, &SAdvancedRotationInputBox::HandleOnToggleChanged, ESlateTransformSubComponent::X)
+				.OnToggleYChanged(this, &SAdvancedRotationInputBox::HandleOnToggleChanged, ESlateTransformSubComponent::Y)
+				.OnToggleZChanged(this, &SAdvancedRotationInputBox::HandleOnToggleChanged, ESlateTransformSubComponent::Z)
+				.TogglePadding(InArgs._TogglePadding)
+				.Visibility(this, &SAdvancedRotationInputBox::IsEulerInputBoxVisible)
+			]
+
+			// a vector 4d variation - to represent quaternion
+			+SHorizontalBox::Slot()
+			.FillWidth(1.f)
+			.Padding(FMargin(0.f))
+			[
+				SNew(SNumericVectorInputBox4)
+				.bColorAxisLabels(InArgs._bColorAxisLabels)
+				.Font(InArgs._Font)
+				.AllowSpin(InArgs._AllowSpin)
+				.SpinDelta(0.001f)
+				.MinSliderVector(FVector4(-1.f, -1.f, -1.f, -1.f))
+				.MaxSliderVector(FVector4(1.f, 1.f, 1.f, 1.f))
+				.X(this, &SAdvancedRotationInputBox::HandleGetNumericValue, ESlateTransformSubComponent::X)
+				.Y(this, &SAdvancedRotationInputBox::HandleGetNumericValue, ESlateTransformSubComponent::Y)
+				.Z(this, &SAdvancedRotationInputBox::HandleGetNumericValue, ESlateTransformSubComponent::Z)
+				.W(this, &SAdvancedRotationInputBox::HandleGetNumericValue, ESlateTransformSubComponent::W)
+				.OnXChanged(this, &SAdvancedRotationInputBox::HandleOnNumericValueChanged, ESlateTransformSubComponent::X)
+				.OnYChanged(this, &SAdvancedRotationInputBox::HandleOnNumericValueChanged, ESlateTransformSubComponent::Y)
+				.OnZChanged(this, &SAdvancedRotationInputBox::HandleOnNumericValueChanged, ESlateTransformSubComponent::Z)
+				.OnWChanged(this, &SAdvancedRotationInputBox::HandleOnNumericValueChanged, ESlateTransformSubComponent::W)
+				.OnXCommitted(this, &SAdvancedRotationInputBox::HandleOnNumericValueCommitted, ESlateTransformSubComponent::X)
+				.OnYCommitted(this, &SAdvancedRotationInputBox::HandleOnNumericValueCommitted, ESlateTransformSubComponent::Y)
+				.OnZCommitted(this, &SAdvancedRotationInputBox::HandleOnNumericValueCommitted, ESlateTransformSubComponent::Z)
+				.OnWCommitted(this, &SAdvancedRotationInputBox::HandleOnNumericValueCommitted, ESlateTransformSubComponent::W)
+				.DisplayToggle(InArgs._DisplayToggle)
+				.ToggleXChecked(this, &SAdvancedRotationInputBox::HandleGetToggleChecked, ESlateTransformSubComponent::X)
+				.ToggleYChecked(this, &SAdvancedRotationInputBox::HandleGetToggleChecked, ESlateTransformSubComponent::Y)
+				.ToggleZChecked(this, &SAdvancedRotationInputBox::HandleGetToggleChecked, ESlateTransformSubComponent::Z)
+				.ToggleWChecked(this, &SAdvancedRotationInputBox::HandleGetToggleChecked, ESlateTransformSubComponent::W)
+				.OnToggleXChanged(this, &SAdvancedRotationInputBox::HandleOnToggleChanged, ESlateTransformSubComponent::X)
+				.OnToggleYChanged(this, &SAdvancedRotationInputBox::HandleOnToggleChanged, ESlateTransformSubComponent::Y)
+				.OnToggleZChanged(this, &SAdvancedRotationInputBox::HandleOnToggleChanged, ESlateTransformSubComponent::Z)
+				.OnToggleWChanged(this, &SAdvancedRotationInputBox::HandleOnToggleChanged, ESlateTransformSubComponent::W)
+				.TogglePadding(InArgs._TogglePadding)
+				.Visibility(this, &SAdvancedRotationInputBox::IsQuaternionInputBoxVisible)
+			]
+
+			/*
+			// a numeric input for the angle for the axis and angle representation
+			+SHorizontalBox::Slot()
+			.Padding(FMargin(0.f))
+			.FillWidth(0.333f)
+			[
+				SNew(SNumericEntryBox<NumericType>)
+				.Font(InArgs._Font)
+				.AllowSpin(InArgs._AllowSpin)
+				.Delta(0.01f)
+				.MinSliderValue(-360.f)
+
+				.Visibility(this, &SAdvancedRotationInputBox::IsAxisAndAngleInputBoxVisible)
+			]
+			*/
+
+			// a vector 3 variation to represent the axis for axis and angle
+			+SHorizontalBox::Slot()
+			.Padding(FMargin(0.f))
+			[
+				SNew(SNumericVectorInputBox3)
+				.bColorAxisLabels(InArgs._bColorAxisLabels)
+				.Font(InArgs._Font)
+				.AllowSpin(InArgs._AllowSpin)
+				.SpinDelta(0.01f)
+				.MinSliderVector(FVector(-1.f, -1.f, -1.f))
+				.MaxSliderVector(FVector(1.f, 1.f, 1.f))
+				.X(this, &SAdvancedRotationInputBox::HandleGetNumericValue, ESlateTransformSubComponent::X)
+				.Y(this, &SAdvancedRotationInputBox::HandleGetNumericValue, ESlateTransformSubComponent::Y)
+				.Z(this, &SAdvancedRotationInputBox::HandleGetNumericValue, ESlateTransformSubComponent::Z)
+				.OnXChanged(this, &SAdvancedRotationInputBox::HandleOnNumericValueChanged, ESlateTransformSubComponent::X)
+				.OnYChanged(this, &SAdvancedRotationInputBox::HandleOnNumericValueChanged, ESlateTransformSubComponent::Y)
+				.OnZChanged(this, &SAdvancedRotationInputBox::HandleOnNumericValueChanged, ESlateTransformSubComponent::Z)
+				.OnXCommitted(this, &SAdvancedRotationInputBox::HandleOnNumericValueCommitted, ESlateTransformSubComponent::X)
+				.OnYCommitted(this, &SAdvancedRotationInputBox::HandleOnNumericValueCommitted, ESlateTransformSubComponent::Y)
+				.OnZCommitted(this, &SAdvancedRotationInputBox::HandleOnNumericValueCommitted, ESlateTransformSubComponent::Z)
+				.DisplayToggle(InArgs._DisplayToggle)
+				.ToggleXChecked(this, &SAdvancedRotationInputBox::HandleGetToggleChecked, ESlateTransformSubComponent::X)
+				.ToggleYChecked(this, &SAdvancedRotationInputBox::HandleGetToggleChecked, ESlateTransformSubComponent::Y)
+				.ToggleZChecked(this, &SAdvancedRotationInputBox::HandleGetToggleChecked, ESlateTransformSubComponent::Z)
+				.OnToggleXChanged(this, &SAdvancedRotationInputBox::HandleOnToggleChanged, ESlateTransformSubComponent::X)
+				.OnToggleYChanged(this, &SAdvancedRotationInputBox::HandleOnToggleChanged, ESlateTransformSubComponent::Y)
+				.OnToggleZChanged(this, &SAdvancedRotationInputBox::HandleOnToggleChanged, ESlateTransformSubComponent::Z)
+				.TogglePadding(InArgs._TogglePadding)
+				.Visibility(this, &SAdvancedRotationInputBox::IsAxisAndAngleInputBoxVisible)
+			]
+
+			// a numeric input for the angle for the axis and angle representation
+			+SHorizontalBox::Slot()
+			.Padding(FMargin(0.f))
+			.FillWidth(0.333f)
+			[
+				SNew(SNumericEntryBox<NumericType>)
+				.Font(InArgs._Font)
+				.AllowSpin(InArgs._AllowSpin)
+				.Delta(0.01f)
+				.MinSliderValue(-360.f)
+				.MaxSliderValue(360.f)
+				.Value(this, &SAdvancedRotationInputBox::HandleGetNumericValue, ESlateTransformSubComponent::Angle)
+				.OnValueChanged(this, &SAdvancedRotationInputBox::HandleOnNumericValueChanged, ESlateTransformSubComponent::Angle)
+				.OnValueCommitted(this, &SAdvancedRotationInputBox::HandleOnNumericValueCommitted, ESlateTransformSubComponent::Angle)
+				.ToolTipText_Lambda([this]()
+				{
+					TOptional<NumericType> Value = HandleGetNumericValue(ESlateTransformSubComponent::Angle);
+					if (Value.IsSet())
+					{
+						return FText::Format(NSLOCTEXT("SAdvancedRotationInputBox", "Angle_ToolTip", "Angle: {0}"), Value.GetValue());
+					}
+					return NSLOCTEXT("SAdvancedRotationInputBox", "MultipleValues", "Multiple Values");
+				})
+				.UndeterminedString(NSLOCTEXT("SAdvancedRotationInputBox", "MultipleValues", "Multiple Values"))
+				.TypeInterface(DegreesTypeInterface)
+				.LabelPadding(FMargin(3))
+				.LabelLocation(SNumericEntryBox<NumericType>::ELabelLocation::Inside)
+				.Label()
+				[
+					AngleLabelWidget
+				]
+				.DisplayToggle(InArgs._DisplayToggle)
+				.TogglePadding(InArgs._TogglePadding)
+				.ToggleChecked(this, &SAdvancedRotationInputBox::HandleGetToggleChecked, ESlateTransformSubComponent::Angle)
+				.OnToggleChanged(this, &SAdvancedRotationInputBox::HandleOnToggleChanged, ESlateTransformSubComponent::Angle)
+				.Visibility(this, &SAdvancedRotationInputBox::IsAxisAndAngleInputBoxVisible)
+			]
+		];
+
+	}
+
+	static bool IsValidComponent(ESlateRotationRepresentation::Type InRepresentation, ESlateTransformSubComponent::Type InSubComponent)
+	{
+		switch(InRepresentation)
+		{
+			case ESlateRotationRepresentation::EulerXYZ:
+			case ESlateRotationRepresentation::EulerXZY:
+			case ESlateRotationRepresentation::EulerYXZ:
+			case ESlateRotationRepresentation::EulerYZX:
+			case ESlateRotationRepresentation::EulerZXY:
+			case ESlateRotationRepresentation::EulerZYX:
+			{
+				switch(InSubComponent)
+				{
+					case ESlateTransformSubComponent::X:
+					case ESlateTransformSubComponent::Y:
+					case ESlateTransformSubComponent::Z:
+					{
+						return true;
+					}
+				}
+				break;
+			}
+			case ESlateRotationRepresentation::Rotator:
+			{
+				switch(InSubComponent)
+				{
+					case ESlateTransformSubComponent::Roll:
+					case ESlateTransformSubComponent::Yaw:
+					case ESlateTransformSubComponent::Pitch:
+					{
+						return true;
+					}
+				}
+				break;
+			}
+			case ESlateRotationRepresentation::Quaternion:
+			{
+				switch(InSubComponent)
+				{
+					case ESlateTransformSubComponent::X:
+					case ESlateTransformSubComponent::Y:
+					case ESlateTransformSubComponent::Z:
+					case ESlateTransformSubComponent::W:
+					{
+						return true;
+					}
+				}
+				break;
+			}
+			case ESlateRotationRepresentation::AxisAndAngle:
+			{
+				switch(InSubComponent)
+				{
+					case ESlateTransformSubComponent::X:
+					case ESlateTransformSubComponent::Y:
+					case ESlateTransformSubComponent::Z:
+					case ESlateTransformSubComponent::Angle:
+					{
+						return true;
+					}
+				}
+				break;
+			}
+		}
+
+		return false;
+	}
+
+private:
+
+	TAttribute<ESlateRotationRepresentation::Type> Representation;
+	TAttribute<TOptional<FRotator>> Rotator;
+	FOnRotatorChanged OnRotatorChanged;
+	FOnRotatorCommitted OnRotatorCommitted;
+	TAttribute<TOptional<FQuat>> Quaternion;
+	FOnQuaternionChanged OnQuaternionChanged;
+	FOnQuaternionCommitted OnQuaternionCommitted;
+	FOnGetNumericValue OnGetNumericValue;
+	FOnNumericValueChanged OnNumericValueChanged;
+	FOnNumericValueCommitted OnNumericValueCommitted;
+	FOnGetToggleChecked OnGetToggleChecked;
+	FOnToggleChanged OnToggleChanged;
+
+	TOptional<NumericType> HandleGetNumericValue(ESlateTransformSubComponent::Type InSubComponent) const
+	{
+		const ESlateRotationRepresentation::Type& Repr = Representation.Get(); 
+		if(IsValidComponent(Repr, InSubComponent))
+		{
+			if(OnGetNumericValue.IsBound())
+			{
+				return OnGetNumericValue.Execute(Repr, InSubComponent);
+			}
+			else if(Rotator.IsBound() || Quaternion.IsBound())
+			{
+				TOptional<FRotator> Rot;
+				TOptional<FQuat> Quat;
+				GetRotatorAndQuat(Rot, Quat);
+
+				if(Rot.IsSet() && Quat.IsSet())
+				{
+					switch(Repr)
+					{
+						case ESlateRotationRepresentation::EulerXYZ:
+						case ESlateRotationRepresentation::EulerXZY:
+						case ESlateRotationRepresentation::EulerYXZ:
+						case ESlateRotationRepresentation::EulerYZX:
+						case ESlateRotationRepresentation::EulerZXY:
+						case ESlateRotationRepresentation::EulerZYX:
+						{
+								const int32 RotationOrder = int32(Repr) - int32(ESlateRotationRepresentation::EulerXYZ);
+								const FVector Euler = AnimationCore::EulerFromQuat(Quat.GetValue(), EEulerRotationOrder(RotationOrder));
+								switch(InSubComponent)
+								{
+									case ESlateTransformSubComponent::X:
+									{
+										return Euler.X;
+									}
+									case ESlateTransformSubComponent::Y:
+									{
+										return Euler.Y;
+									}
+									case ESlateTransformSubComponent::Z:
+									{
+										return Euler.Z;
+									}
+									default:
+									{
+										break;
+									}
+								}
+								break;
+						}
+						case ESlateRotationRepresentation::Rotator:
+						{
+							switch(InSubComponent)
+							{
+								case ESlateTransformSubComponent::Pitch:
+								{
+									return Rot->Pitch;
+								}
+								case ESlateTransformSubComponent::Yaw:
+								{
+									return Rot->Yaw;
+								}
+								case ESlateTransformSubComponent::Roll:
+								{
+									return Rot->Roll;
+								}
+								default:
+								{
+									break;
+								}
+							}
+							break;
+						}
+						case ESlateRotationRepresentation::Quaternion:
+						{
+							switch(InSubComponent)
+							{
+								case ESlateTransformSubComponent::X:
+								{
+									return Quat->X;
+								}
+								case ESlateTransformSubComponent::Y:
+								{
+									return Quat->Y;
+								}
+								case ESlateTransformSubComponent::Z:
+								{
+									return Quat->Z;
+								}
+								case ESlateTransformSubComponent::W:
+								{
+									return Quat->W;
+								}
+								default:
+								{
+									break;
+								}
+							}
+							break;
+						}
+						case ESlateRotationRepresentation::AxisAndAngle:
+						{
+							FVector Axis;
+							double Angle;
+							Quat->ToAxisAndAngle(Axis, Angle);
+							switch(InSubComponent)
+							{
+								case ESlateTransformSubComponent::X:
+								{
+									return Axis.X;
+								}
+								case ESlateTransformSubComponent::Y:
+								{
+									return Axis.Y;
+								}
+								case ESlateTransformSubComponent::Z:
+								{
+									return Axis.Z;
+								}
+								case ESlateTransformSubComponent::Angle:
+								{
+									return Angle;
+								}
+								default:
+								{
+									break;
+								}
+							}
+							break;
+						}
+					}
+				}
+			}
+		}
+		return TOptional<NumericType>();
+	}
+
+	bool GetRotatorAndQuat(TOptional<FRotator>& Rot, TOptional<FQuat> &Quat) const
+	{
+		bool bSuccess = false;
+		if(Rotator.IsBound())
+		{
+			Rot = Rotator.Get();
+			if(Rot.IsSet())
+			{
+				Quat = Rot->Quaternion();
+				bSuccess = true;
+			}
+		}
+				
+		if(Quaternion.IsBound())
+		{
+			Quat = Quaternion.Get();
+			if(Quat.IsSet())
+			{
+				if(!Rot.IsSet())
+				{
+					Rot = Quat->Rotator();
+				}
+				bSuccess = true;;
+			}
+			else if(Rot.IsSet())
+			{
+				Quat = Rot->Quaternion();
+				bSuccess = true;
+			}
+		}
+		return bSuccess;
+	}
+
+	bool ApplyNumericValueChangeToRotatorAndQuat(TOptional<FRotator>& Rot, TOptional<FQuat> &Quat, NumericType InValue, ESlateTransformSubComponent::Type InSubComponent) const
+	{
+		if(!GetRotatorAndQuat(Rot, Quat))
+		{
+			return false;
+		}
+		
+		if(Rot.IsSet() && Quat.IsSet())
+		{
+			const ESlateRotationRepresentation::Type& Repr = Representation.Get();
+			switch(Repr)
+			{
+				case ESlateRotationRepresentation::EulerXYZ:
+				case ESlateRotationRepresentation::EulerXZY:
+				case ESlateRotationRepresentation::EulerYXZ:
+				case ESlateRotationRepresentation::EulerYZX:
+				case ESlateRotationRepresentation::EulerZXY:
+				case ESlateRotationRepresentation::EulerZYX:
+				{
+						const int32 RotationOrder = int32(Repr) - int32(ESlateRotationRepresentation::EulerXYZ);
+						FVector Euler = AnimationCore::EulerFromQuat(Quat.GetValue(), EEulerRotationOrder(RotationOrder));
+						switch(InSubComponent)
+						{
+							case ESlateTransformSubComponent::X:
+							{
+								Euler.X = InValue;
+								break;
+							}
+							case ESlateTransformSubComponent::Y:
+							{
+								Euler.Y = InValue;
+								break;
+							}
+							case ESlateTransformSubComponent::Z:
+							{
+								Euler.Z = InValue;
+								break;
+							}
+							default:
+							{
+								break;
+							}
+						}
+						Quat = AnimationCore::QuatFromEuler(Euler, EEulerRotationOrder(RotationOrder));
+						Rot = Quat->Rotator();
+						break;
+				}
+				case ESlateRotationRepresentation::Rotator:
+				{
+					switch(InSubComponent)
+					{
+						case ESlateTransformSubComponent::Pitch:
+						{
+							Rot->Pitch = InValue;
+							break;
+						}
+						case ESlateTransformSubComponent::Yaw:
+						{
+							Rot->Yaw = InValue;
+							break;
+						}
+						case ESlateTransformSubComponent::Roll:
+						{
+							Rot->Roll = InValue;
+							break;
+						}
+						default:
+						{
+							break;
+						}
+					}
+					Quat = Rot->Quaternion();
+					break;
+				}
+				case ESlateRotationRepresentation::Quaternion:
+				{
+					switch(InSubComponent)
+					{
+						case ESlateTransformSubComponent::X:
+						{
+							Quat->X = InValue;
+							break;
+						}
+						case ESlateTransformSubComponent::Y:
+						{
+							Quat->Y = InValue;
+							break;
+						}
+						case ESlateTransformSubComponent::Z:
+						{
+							Quat->Z = InValue;
+							break;
+						}
+						case ESlateTransformSubComponent::W:
+						{
+							Quat->W = InValue;
+							break;
+						}
+						default:
+						{
+							break;
+						}
+					}
+					Rot = Quat->Rotator();
+					break;
+				}
+				case ESlateRotationRepresentation::AxisAndAngle:
+				{
+					FVector Axis;
+					double Angle;
+					Quat->ToAxisAndAngle(Axis, Angle);
+					switch(InSubComponent)
+					{
+						case ESlateTransformSubComponent::X:
+						{
+							Axis.X = InValue;
+							break;
+						}
+						case ESlateTransformSubComponent::Y:
+						{
+							Axis.Y = InValue;
+							break;
+						}
+						case ESlateTransformSubComponent::Z:
+						{
+							Axis.Z = InValue;
+							break;
+						}
+						case ESlateTransformSubComponent::Angle:
+						{
+							Angle = InValue;
+							break;
+						}
+						default:
+						{
+							break;
+						}
+					}
+
+					Quat = FQuat(Axis, Angle);
+					Rot = Quat->Rotator();
+					break;
+				}
+			}
+		}
+
+		return true;
+	}
+
+	void HandleOnNumericValueChanged(NumericType InValue, ESlateTransformSubComponent::Type InSubComponent)
+	{
+		const ESlateRotationRepresentation::Type& Repr = Representation.Get(); 
+		if(IsValidComponent(Repr, InSubComponent))
+		{
+			if(OnNumericValueChanged.IsBound())
+			{
+				OnNumericValueChanged.Execute(Repr, InSubComponent, InValue);
+			}
+
+			if(OnRotatorChanged.IsBound() || OnQuaternionChanged.IsBound())
+			{
+				TOptional<FRotator> Rot;
+				TOptional<FQuat> Quat;
+				if(ApplyNumericValueChangeToRotatorAndQuat(Rot, Quat, InValue, InSubComponent))
+				{
+					if(OnRotatorChanged.IsBound() && Rot.IsSet())
+					{
+						OnRotatorChanged.Execute(Rot.GetValue());
+					}
+					if(OnQuaternionChanged.IsBound() && Quat.IsSet())
+					{
+						OnQuaternionChanged.Execute(Quat.GetValue());
+					}
+				}
+			}
+		}
+	}
+	
+	void HandleOnNumericValueCommitted(NumericType InValue, ETextCommit::Type InCommitType, ESlateTransformSubComponent::Type InSubComponent)
+	{
+		const ESlateRotationRepresentation::Type& Repr = Representation.Get(); 
+		if(IsValidComponent(Repr, InSubComponent))
+		{
+			if(OnNumericValueCommitted.IsBound())
+			{
+				OnNumericValueCommitted.Execute(Repr, InSubComponent, InValue, InCommitType);
+			}
+			if(OnRotatorCommitted.IsBound() || OnQuaternionCommitted.IsBound())
+			{
+				TOptional<FRotator> Rot;
+				TOptional<FQuat> Quat;
+				if(ApplyNumericValueChangeToRotatorAndQuat(Rot, Quat, InValue, InSubComponent))
+				{
+					if(OnRotatorCommitted.IsBound() && Rot.IsSet())
+					{
+						OnRotatorCommitted.Execute(Rot.GetValue(), InCommitType);
+					}
+					if(OnQuaternionCommitted.IsBound() && Quat.IsSet())
+					{
+						OnQuaternionCommitted.Execute(Quat.GetValue(), InCommitType);
+					}
+				}
+			}
+		}
+	}
+
+	ECheckBoxState HandleGetToggleChecked(ESlateTransformSubComponent::Type InSubComponent) const
+	{
+		const ESlateRotationRepresentation::Type& Repr = Representation.Get(); 
+		if(IsValidComponent(Repr, InSubComponent))
+		{
+			if(OnGetToggleChecked.IsBound())
+			{
+				return OnGetToggleChecked.Execute(Repr, InSubComponent);
+			}
+		}
+		return ECheckBoxState::Checked;
+	}
+
+	void HandleOnToggleChanged(ECheckBoxState InValue, ESlateTransformSubComponent::Type InSubComponent)
+	{
+		const ESlateRotationRepresentation::Type& Repr = Representation.Get(); 
+		if(IsValidComponent(Repr, InSubComponent))
+		{
+			if(OnToggleChanged.IsBound())
+			{
+				OnToggleChanged.Execute(Repr, InSubComponent, InValue);
+			}
+		}
+	}
+
+	EVisibility IsRotatorInputBoxVisible() const
+	{
+		const ESlateRotationRepresentation::Type& Repr = Representation.Get();
+		return Repr == ESlateRotationRepresentation::Rotator ? EVisibility::Visible : EVisibility::Collapsed;
+	}
+
+	EVisibility IsEulerInputBoxVisible() const
+	{
+		const ESlateRotationRepresentation::Type& Repr = Representation.Get();
+		switch(Repr)
+		{
+			case ESlateRotationRepresentation::EulerXYZ:
+			case ESlateRotationRepresentation::EulerXZY:
+			case ESlateRotationRepresentation::EulerYXZ:
+			case ESlateRotationRepresentation::EulerYZX:
+			case ESlateRotationRepresentation::EulerZXY:
+			case ESlateRotationRepresentation::EulerZYX:
+			{
+				return EVisibility::Visible;
+			}
+		}
+		return EVisibility::Collapsed;
+	}
+
+	EVisibility IsQuaternionInputBoxVisible() const
+	{
+		const ESlateRotationRepresentation::Type& Repr = Representation.Get();
+		switch(Repr)
+		{
+			case ESlateRotationRepresentation::Quaternion:
+			{
+				return EVisibility::Visible;
+			}
+		}
+		return EVisibility::Collapsed;
+	}
+
+	EVisibility IsAxisAndAngleInputBoxVisible() const
+	{
+		const ESlateRotationRepresentation::Type& Repr = Representation.Get();
+		switch(Repr)
+		{
+			case ESlateRotationRepresentation::AxisAndAngle:
+			{
+				return EVisibility::Visible;
+			}
+		}
+		return EVisibility::Collapsed;
+	}
+};
+
diff --git a/Engine/Source/Developer/AnimationWidgets/Public/SAdvancedTransformInputBox.h b/Engine/Source/Developer/AnimationWidgets/Public/SAdvancedTransformInputBox.h
new file mode 100644
--- /dev/null
+++ b/Engine/Source/Developer/AnimationWidgets/Public/SAdvancedTransformInputBox.h
@@ -0,0 +1,1074 @@
+// Copyright Epic Games, Inc. All Rights Reserved.
+
+#pragma once
+
+#include "SAdvancedRotationInputBox.h"
+#include "Widgets/Input/SComboButton.h"
+#include "Styling/AppStyle.h"
+#include "Framework/Application/SlateApplication.h"
+
+namespace ESlateTransformComponent
+{
+	enum Type
+	{
+		Location,
+		Rotation,
+		Scale,
+		Max
+	};
+}
+
+/**
+ * Generic Transform Slate control
+ */
+template<typename TransformType = FTransform, typename NumericType = FVector::FReal>
+class SAdvancedTransformInputBox : public SCompoundWidget
+{
+public:
+
+	typedef SNumericVectorInputBox<NumericType, UE::Math::TVector<NumericType>, 3> SNumericVectorInputBox3; 
+
+	/** Delegate for notification for a transform value change */
+	DECLARE_DELEGATE_OneParam(FOnTransformChanged, TransformType);
+
+	/** Delegate for notification for a transform value commit */
+	DECLARE_DELEGATE_TwoParams(FOnTransformCommitted, TransformType, ETextCommit::Type);
+
+	/** Delegate to retrieve a numerictype value  */
+	DECLARE_DELEGATE_RetVal_ThreeParams(TOptional<NumericType>, FOnGetNumericValue, ESlateTransformComponent::Type, ESlateRotationRepresentation::Type, ESlateTransformSubComponent::Type);
+
+	/** Notification for numerictype value change */
+ 	DECLARE_DELEGATE_FourParams(FOnNumericValueChanged, ESlateTransformComponent::Type, ESlateRotationRepresentation::Type, ESlateTransformSubComponent::Type, NumericType);
+
+	/** Notification for numerictype value committed */
+	DECLARE_DELEGATE_FiveParams(FOnNumericValueCommitted, ESlateTransformComponent::Type, ESlateRotationRepresentation::Type, ESlateTransformSubComponent::Type, NumericType, ETextCommit::Type);
+
+	/** Delegate to retrieve relative / world state */
+	DECLARE_DELEGATE_RetVal_OneParam(bool, FOnGetIsComponentRelative, ESlateTransformComponent::Type);
+
+	/** Notification for a relative / world state change */
+	DECLARE_DELEGATE_TwoParams(FOnIsComponentRelativeChanged, ESlateTransformComponent::Type, bool /* is relative? */);
+
+	/** Notification for the rotation representation change */
+	DECLARE_DELEGATE_OneParam(FOnRotationRepresentationChanged, ESlateRotationRepresentation::Type);
+
+	/** Notification for the scale lock change */
+	DECLARE_DELEGATE_OneParam(FOnScaleLockChanged, bool /* is locked? */);
+
+	/** Delegate to retrieve toggle checkstate */
+	DECLARE_DELEGATE_RetVal_ThreeParams(ECheckBoxState, FOnGetToggleChecked, ESlateTransformComponent::Type, ESlateRotationRepresentation::Type, ESlateTransformSubComponent::Type);
+
+	/** Notification for toggle checkstate change */
+	DECLARE_DELEGATE_FourParams(FOnToggleChanged, ESlateTransformComponent::Type, ESlateRotationRepresentation::Type, ESlateTransformSubComponent::Type, ECheckBoxState);
+
+	SLATE_BEGIN_ARGS( SAdvancedTransformInputBox )
+		: _ConstructLocation(true)
+		, _ConstructRotation(true)
+		, _ConstructScale(true)
+		, _UseQuaternionForRotation(false)
+		, _RotationRepresentation(TSharedPtr<ESlateRotationRepresentation::Type>( new ESlateRotationRepresentation::Type(ESlateRotationRepresentation::Rotator)))
+		, _bColorAxisLabels(true)
+		, _ShowInlineLabels(false) 
+		, _LocationLabel(NSLOCTEXT("SAdvancedTransformInputBox", "Location", "Location"))
+		, _RotationLabel(NSLOCTEXT("SAdvancedTransformInputBox", "Rotation", "Rotation"))
+		, _ScaleLabel(NSLOCTEXT("SAdvancedTransformInputBox", "Scale", "Scale"))
+		, _LabelPadding( FMargin( 0.f,0.f,6.f,0.f ) )
+		, _Font(FAppStyle::Get().GetFontStyle("NormalFont"))
+		, _AllowSpin(true)
+		, _SpinDelta(0.001f)
+		, _AllowEditRotationRepresentation(true)
+		, _DisplayScaleLock(true)
+		, _IsScaleLocked(TSharedPtr<bool>(new bool(false)))
+		, _DisplayRelativeWorld(false)
+		, _DisplayToggle( false )
+		, _TogglePadding( FMargin( 1.f,0.f,1.f,0.f ) )
+		{}
+
+		/** Whether or not to construct the location component */
+		SLATE_ARGUMENT( bool, ConstructLocation )
+
+		/** Whether or not to construct the rotation component */
+		SLATE_ARGUMENT( bool, ConstructRotation )
+
+		/** Whether or not to construct the scale component */
+		SLATE_ARGUMENT( bool, ConstructScale )
+
+		/** Whether to use the quaternion is the primary rotation representation or the rotator (default) */
+		SLATE_ARGUMENT( bool, UseQuaternionForRotation )
+
+		/** Optional representation as transform */
+		SLATE_ATTRIBUTE( TOptional<TransformType>, Transform )
+
+		/** Optional delegate to notify a change in the transform */
+		SLATE_EVENT( FOnTransformChanged, OnTransformChanged )
+
+		/** Optional delegate to notify a commit in the transform */
+		SLATE_EVENT( FOnTransformCommitted, OnTransformCommitted )
+
+		/** Representation of the rotation */
+		SLATE_ARGUMENT( TSharedPtr<ESlateRotationRepresentation::Type>, RotationRepresentation )
+
+		/** Delegate to notify a rotation representation change */
+		SLATE_EVENT( FOnRotationRepresentationChanged, OnRotationRepresentationChanged )
+
+		/** Delegate to retrieve a value */
+		SLATE_EVENT( FOnGetNumericValue, OnGetNumericValue )
+
+		/** Delegate to notify a value change */
+		SLATE_EVENT( FOnNumericValueChanged, OnNumericValueChanged )
+
+		/** Delegate to notify a committed value */
+		SLATE_EVENT( FOnNumericValueCommitted, OnNumericValueCommitted )
+
+		/** Should the axis labels be colored */
+		SLATE_ARGUMENT( bool, bColorAxisLabels )		
+
+		/** Should we show labels inline next to all input boxes */
+		SLATE_ARGUMENT( bool, ShowInlineLabels )		
+
+		/** The label to use for the location component */
+		SLATE_ARGUMENT( FText, LocationLabel )		
+
+		/** The label to use for the rotation component */
+		SLATE_ARGUMENT( FText, RotationLabel )		
+
+		/** The label to use for the scale component */
+		SLATE_ARGUMENT( FText, ScaleLabel )		
+
+		/** Padding around the label checkbox */
+		SLATE_ARGUMENT( FMargin, LabelPadding )
+
+		/** Font to use for the text in this box */
+		SLATE_ATTRIBUTE( FSlateFontInfo, Font )
+
+		/** Whether or not values can be spun or if they should be typed in */
+		SLATE_ARGUMENT( bool, AllowSpin )
+
+		/** The delta amount to apply, per pixel, when the spinner is dragged. */
+		SLATE_ATTRIBUTE( NumericType, SpinDelta )
+
+		/** Whether or not to display the rotation representation picker in the label */
+		SLATE_ARGUMENT( bool, AllowEditRotationRepresentation )
+	
+		/** Whether or not to display the lock scale button */
+		SLATE_ARGUMENT( bool, DisplayScaleLock )
+
+		/** The state of the scale lock */
+		SLATE_ARGUMENT( TSharedPtr<bool>, IsScaleLocked )
+
+		/** Delegate to notify a scale lock change */
+		SLATE_EVENT( FOnScaleLockChanged, OnScaleLockChanged )
+
+		/** Whether or not to display the per component buttons to switch between relative and world */
+		SLATE_ARGUMENT( bool, DisplayRelativeWorld )
+
+		/** Delegate to retrieve a relative / world state per component */
+		SLATE_EVENT( FOnGetIsComponentRelative, OnGetIsComponentRelative )
+
+		/** Delegate to notify a relative / world state change */
+		SLATE_EVENT( FOnIsComponentRelativeChanged, OnIsComponentRelativeChanged )
+
+		/** Whether or not to include a toggle checkbox to the left of the widget */
+		SLATE_ARGUMENT( bool, DisplayToggle )
+
+		/** Delegate to retrieve a toggle state */
+		SLATE_EVENT( FOnGetToggleChecked, OnGetToggleChecked )
+
+		/** Delegate to notify a toggle state change */
+		SLATE_EVENT( FOnToggleChanged, OnToggleChanged )
+
+		/** Padding around the toggle checkbox */
+		SLATE_ARGUMENT( FMargin, TogglePadding )
+
+	SLATE_END_ARGS()
+
+	/**
+	 * Construct this widget
+	 *
+	 * @param	InArgs	The declaration data for this widget
+	 */
+	void Construct( const FArguments& InArgs )
+	{
+		TSharedRef<SVerticalBox> VerticalBox = SNew(SVerticalBox);
+
+		if(InArgs._ConstructLocation)
+		{
+			VerticalBox->AddSlot()
+			[
+				ConstructWidget(InArgs, ESlateTransformComponent::Location)
+			];
+		}
+		if(InArgs._ConstructRotation)
+		{
+			VerticalBox->AddSlot()
+			[
+				ConstructWidget(InArgs, ESlateTransformComponent::Rotation)
+			];
+		}
+		if(InArgs._ConstructScale)
+		{
+			VerticalBox->AddSlot()
+			[
+				ConstructWidget(InArgs, ESlateTransformComponent::Scale)
+			];
+		}
+
+		ChildSlot
+		[
+			VerticalBox
+		];
+	}
+
+	/**
+	 * Construct an input widget
+	 */
+	static TSharedRef<SWidget> ConstructWidget(const FArguments& InArgs, ESlateTransformComponent::Type InComponent)
+	{
+		TSharedRef<SHorizontalBox> HorizontalBox = SNew(SHorizontalBox);
+
+		if(InArgs._ShowInlineLabels)
+		{
+			HorizontalBox->AddSlot()
+			.HAlign(HAlign_Left)
+			.VAlign(VAlign_Center)
+			.Padding(InArgs._LabelPadding)
+			.AutoWidth()
+			[
+				ConstructLabel(InArgs, InComponent)
+			];
+		}
+
+		TSharedRef<SWidget> InputWidget = SNullWidget::NullWidget;
+		TAttribute<TOptional<TransformType>> Transform = InArgs._Transform;
+		FOnTransformChanged OnTransformChanged = InArgs._OnTransformChanged;
+		FOnTransformCommitted OnTransformCommitted = InArgs._OnTransformCommitted;
+		FOnGetNumericValue OnGetNumericValue = InArgs._OnGetNumericValue;
+		FOnNumericValueChanged OnNumericValueChanged = InArgs._OnNumericValueChanged;
+		FOnNumericValueCommitted OnNumericValueCommitted = InArgs._OnNumericValueCommitted;
+		FOnGetToggleChecked OnGetToggleChecked = InArgs._OnGetToggleChecked;
+		FOnToggleChanged OnToggleChanged = InArgs._OnToggleChanged;
+		const bool UseQuaternionForRotation = InArgs._UseQuaternionForRotation;
+
+		auto OnGetLocation = [Transform, OnGetNumericValue]() -> TOptional<FVector>
+		{
+			if(Transform.IsBound())
+			{
+				TOptional<TransformType> Xfo = Transform.Get();
+				if(Xfo.IsSet())
+				{
+					return Xfo->GetLocation();
+				}
+			}
+			if(OnGetNumericValue.IsBound())
+			{
+				FVector Location = FVector::OneVector;
+				Location.X = OnGetNumericValue.Execute(ESlateTransformComponent::Location, ESlateRotationRepresentation::EulerXYZ, ESlateTransformSubComponent::X).Get(Location.X);
+				Location.Y = OnGetNumericValue.Execute(ESlateTransformComponent::Location, ESlateRotationRepresentation::EulerXYZ, ESlateTransformSubComponent::Y).Get(Location.Y);
+				Location.Z = OnGetNumericValue.Execute(ESlateTransformComponent::Location, ESlateRotationRepresentation::EulerXYZ, ESlateTransformSubComponent::Z).Get(Location.Z);
+				return Location;
+			}
+			return TOptional<FVector>();
+		};
+
+		auto OnGetRotator = [Transform, OnGetNumericValue]() -> TOptional<FRotator>
+		{
+			if(Transform.IsBound())
+			{
+				TOptional<TransformType> Xfo = Transform.Get();
+				if(Xfo.IsSet())
+				{
+					return Xfo->Rotator();
+				}
+			}
+			if(OnGetNumericValue.IsBound())
+			{
+				FRotator Rotator = FRotator::ZeroRotator;
+				Rotator.Roll = OnGetNumericValue.Execute(ESlateTransformComponent::Rotation, ESlateRotationRepresentation::Rotator, ESlateTransformSubComponent::Roll).Get(Rotator.Roll);
+				Rotator.Pitch = OnGetNumericValue.Execute(ESlateTransformComponent::Rotation, ESlateRotationRepresentation::Rotator, ESlateTransformSubComponent::Pitch).Get(Rotator.Pitch);
+				Rotator.Yaw = OnGetNumericValue.Execute(ESlateTransformComponent::Rotation, ESlateRotationRepresentation::Rotator, ESlateTransformSubComponent::Yaw).Get(Rotator.Yaw);
+				return Rotator;
+			}
+			return TOptional<FRotator>();
+		};
+
+		auto OnGetQuaternion = [Transform, OnGetNumericValue]() -> TOptional<FQuat>
+		{
+			if(Transform.IsBound())
+			{
+				TOptional<TransformType> Xfo = Transform.Get();
+				if(Xfo.IsSet())
+				{
+					return Xfo->GetRotation().GetNormalized();
+				}
+			}
+			if(OnGetNumericValue.IsBound())
+			{
+				FQuat Quat = FQuat::Identity;
+				Quat.X = OnGetNumericValue.Execute(ESlateTransformComponent::Rotation, ESlateRotationRepresentation::Quaternion, ESlateTransformSubComponent::X).Get(Quat.X);
+				Quat.Y = OnGetNumericValue.Execute(ESlateTransformComponent::Rotation, ESlateRotationRepresentation::Quaternion, ESlateTransformSubComponent::Y).Get(Quat.Y);
+				Quat.Z = OnGetNumericValue.Execute(ESlateTransformComponent::Rotation, ESlateRotationRepresentation::Quaternion, ESlateTransformSubComponent::Z).Get(Quat.Z);
+				Quat.W = OnGetNumericValue.Execute(ESlateTransformComponent::Rotation, ESlateRotationRepresentation::Quaternion, ESlateTransformSubComponent::W).Get(Quat.W);
+				return Quat.GetNormalized();
+			}
+			return TOptional<FQuat>();
+		};
+
+		auto OnGetScale = [Transform, OnGetNumericValue]() -> TOptional<FVector>
+		{
+			if(Transform.IsBound())
+			{
+				TOptional<TransformType> Xfo = Transform.Get();
+				if(Xfo.IsSet())
+				{
+					return Xfo->GetScale3D();
+				}
+			}
+			if(OnGetNumericValue.IsBound())
+			{
+				FVector Scale = FVector::OneVector;
+				Scale.X = OnGetNumericValue.Execute(ESlateTransformComponent::Scale, ESlateRotationRepresentation::EulerXYZ, ESlateTransformSubComponent::X).Get(Scale.X);
+				Scale.Y = OnGetNumericValue.Execute(ESlateTransformComponent::Scale, ESlateRotationRepresentation::EulerXYZ, ESlateTransformSubComponent::Y).Get(Scale.Y);
+				Scale.Z = OnGetNumericValue.Execute(ESlateTransformComponent::Scale, ESlateRotationRepresentation::EulerXYZ, ESlateTransformSubComponent::Z).Get(Scale.Z);
+				return Scale;
+			}
+			return TOptional<FVector>();
+		};
+
+		auto OnGetVector = [OnGetLocation, OnGetScale](ESlateTransformComponent::Type InComponent) -> TOptional<FVector>
+		{
+			return (InComponent == ESlateTransformComponent::Location) ? OnGetLocation() : OnGetScale();
+		};
+
+		auto OnGetTransform = [OnGetLocation, OnGetQuaternion, OnGetRotator, OnGetScale, UseQuaternionForRotation]() -> TransformType
+		{
+			TransformType Result = TransformType::Identity;
+			Result.SetLocation(OnGetLocation().Get(Result.GetLocation()));
+			Result.SetScale3D(OnGetScale().Get(Result.GetScale3D()));
+
+			if(UseQuaternionForRotation)
+			{
+				Result.SetRotation(OnGetQuaternion().Get(Result.GetRotation()));
+			}
+			else
+			{
+				const FRotator Rotator = OnGetRotator().Get(Result.Rotator());
+				Result = TransformType(Rotator, Result.GetLocation(), Result.GetScale3D());
+			}
+			Result.NormalizeRotation();
+			return Result;
+		};
+
+		switch(InComponent)
+		{
+			case ESlateTransformComponent::Location:
+			case ESlateTransformComponent::Scale:
+			{
+				TAttribute<TOptional<NumericType>> XAttribute, YAttribute, ZAttribute;
+				TAttribute<TOptional<FVector>> Vector3Attribute;
+
+				if(InArgs._OnGetNumericValue.IsBound())
+				{
+					XAttribute = TAttribute<TOptional<NumericType>>::CreateLambda(
+						[OnGetNumericValue, OnGetVector, InComponent]() -> TOptional<NumericType>
+						{
+							return OnGetNumericValue.Execute(InComponent, ESlateRotationRepresentation::EulerXYZ, ESlateTransformSubComponent::X);
+						}); 
+					YAttribute = TAttribute<TOptional<NumericType>>::CreateLambda(
+						[OnGetNumericValue, OnGetVector, InComponent]() -> TOptional<NumericType>
+						{
+							return OnGetNumericValue.Execute(InComponent, ESlateRotationRepresentation::EulerXYZ, ESlateTransformSubComponent::Y);
+						}); 
+					ZAttribute = TAttribute<TOptional<NumericType>>::CreateLambda(
+						[OnGetNumericValue, OnGetVector, InComponent]() -> TOptional<NumericType>
+						{
+							return OnGetNumericValue.Execute(InComponent, ESlateRotationRepresentation::EulerXYZ, ESlateTransformSubComponent::Z);
+						}); 
+				}
+
+				if(InArgs._Transform.IsBound() || InArgs._Transform.IsSet())
+				{
+					Vector3Attribute = TAttribute<TOptional<FVector>>::CreateLambda(
+						[OnGetVector, InComponent]() -> TOptional<FVector>
+						{
+							return OnGetVector(InComponent);
+						});
+				}
+
+				typename SNumericVectorInputBox3::FOnNumericValueChanged XChanged, YChanged, ZChanged;
+				if(OnNumericValueChanged.IsBound())
+				{
+					XChanged = SNumericVectorInputBox3::FOnNumericValueChanged::CreateLambda(
+						[OnNumericValueChanged, InComponent](NumericType InValue)
+						{
+							OnNumericValueChanged.Execute(InComponent, ESlateRotationRepresentation::EulerXYZ, ESlateTransformSubComponent::X, InValue);
+						}
+					); 
+					YChanged = SNumericVectorInputBox3::FOnNumericValueChanged::CreateLambda(
+						[OnNumericValueChanged, InComponent](NumericType InValue)
+						{
+							OnNumericValueChanged.Execute(InComponent, ESlateRotationRepresentation::EulerXYZ, ESlateTransformSubComponent::Y, InValue);
+						}
+					); 
+					ZChanged = SNumericVectorInputBox3::FOnNumericValueChanged::CreateLambda(
+						[OnNumericValueChanged, InComponent](NumericType InValue)
+						{
+							OnNumericValueChanged.Execute(InComponent, ESlateRotationRepresentation::EulerXYZ, ESlateTransformSubComponent::Z, InValue);
+						}
+					); 
+				}
+
+				typename SNumericVectorInputBox3::FOnVectorValueChanged VectorChanged;
+				if(OnTransformChanged.IsBound())
+				{
+					VectorChanged = SNumericVectorInputBox3::FOnVectorValueChanged::CreateLambda(
+						[OnTransformChanged, OnGetTransform, InComponent](FVector InValue)
+						{
+							if(OnTransformChanged.IsBound())
+							{
+								TransformType Transform = OnGetTransform();
+								if(InComponent == ESlateTransformComponent::Location)
+								{
+									Transform.SetLocation(InValue);
+								}
+								else
+								{
+									Transform.SetScale3D(InValue);
+								}
+								OnTransformChanged.Execute(Transform);
+							}
+						}
+					);
+				}
+
+				typename SNumericVectorInputBox3::FOnNumericValueCommitted XCommitted, YCommitted, ZCommitted;
+				if(OnNumericValueCommitted.IsBound())
+				{
+					XCommitted = SNumericVectorInputBox3::FOnNumericValueCommitted::CreateLambda(
+						[OnNumericValueCommitted, InComponent](NumericType InValue, ETextCommit::Type InCommitType)
+						{
+							OnNumericValueCommitted.Execute(InComponent, ESlateRotationRepresentation::EulerXYZ, ESlateTransformSubComponent::X, InValue, InCommitType);
+						}
+					); 
+					YCommitted = SNumericVectorInputBox3::FOnNumericValueCommitted::CreateLambda(
+						[OnNumericValueCommitted, InComponent](NumericType InValue, ETextCommit::Type InCommitType)
+						{
+							OnNumericValueCommitted.Execute(InComponent, ESlateRotationRepresentation::EulerXYZ, ESlateTransformSubComponent::Y, InValue, InCommitType);
+						}
+					); 
+					ZCommitted = SNumericVectorInputBox3::FOnNumericValueCommitted::CreateLambda(
+						[OnNumericValueCommitted, InComponent](NumericType InValue, ETextCommit::Type InCommitType)
+						{
+							OnNumericValueCommitted.Execute(InComponent, ESlateRotationRepresentation::EulerXYZ, ESlateTransformSubComponent::Z, InValue, InCommitType);
+						}
+					); 
+				}
+					
+				typename SNumericVectorInputBox3::FOnVectorValueCommitted VectorCommitted;
+				if(OnTransformCommitted.IsBound())
+				{
+					VectorCommitted = SNumericVectorInputBox3::FOnVectorValueCommitted::CreateLambda(
+						[OnTransformCommitted, OnGetTransform, InComponent](FVector InValue, ETextCommit::Type InCommitType)
+						{
+							if(OnTransformCommitted.IsBound())
+							{
+								TransformType Transform = OnGetTransform();
+								if(InComponent == ESlateTransformComponent::Location)
+								{
+									Transform.SetLocation(InValue);
+								}
+								else
+								{
+									Transform.SetScale3D(InValue);
+								}
+								OnTransformCommitted.Execute(Transform, InCommitType);
+							}
+						}
+					);
+				}
+
+				TSharedPtr<bool> ScaleLockState = InArgs._IsScaleLocked;
+				typename SNumericVectorInputBox3::FOnConstrainVector ConstrainComponents;
+					
+				if(InComponent == ESlateTransformComponent::Scale)
+				{
+					if(ScaleLockState.IsValid())
+					{
+						ConstrainComponents = SNumericVectorInputBox3::FOnConstrainVector::CreateStatic(
+							&SAdvancedTransformInputBox::ConstrainScale, TAttribute<bool>::CreateLambda([ScaleLockState]() 
+							{
+								return *ScaleLockState.Get();
+							})
+						);
+					}
+				}
+
+				InputWidget = SNew(SNumericVectorInputBox3)
+					.Font(InArgs._Font)
+					.AllowSpin(InArgs._AllowSpin)
+					.SpinDelta(InArgs._SpinDelta)
+					.bColorAxisLabels(InArgs._bColorAxisLabels)
+					.X(XAttribute)
+					.Y(YAttribute)
+					.Z(ZAttribute)
+					.Vector(Vector3Attribute)
+					.OnXChanged(XChanged)
+					.OnYChanged(YChanged)
+					.OnZChanged(ZChanged)
+					.OnVectorChanged(VectorChanged)
+					.OnXCommitted(XCommitted)
+					.OnYCommitted(YCommitted)
+					.OnZCommitted(ZCommitted)
+					.OnVectorCommitted(VectorCommitted)
+					.ConstrainVector(ConstrainComponents)
+					.DisplayToggle(InArgs._DisplayToggle)
+					.TogglePadding(InArgs._TogglePadding)
+					.ToggleXChecked_Lambda([OnGetToggleChecked, InComponent]() -> ECheckBoxState
+					{
+						if(OnGetToggleChecked.IsBound())
+						{
+							return OnGetToggleChecked.Execute(InComponent, ESlateRotationRepresentation::EulerXYZ, ESlateTransformSubComponent::X);
+						}
+						return ECheckBoxState::Checked;
+					})
+					.ToggleYChecked_Lambda([OnGetToggleChecked, InComponent]() -> ECheckBoxState
+					{
+						if(OnGetToggleChecked.IsBound())
+						{
+							return OnGetToggleChecked.Execute(InComponent, ESlateRotationRepresentation::EulerXYZ, ESlateTransformSubComponent::Y);
+						}
+						return ECheckBoxState::Checked;
+					})
+					.ToggleZChecked_Lambda([OnGetToggleChecked, InComponent]() -> ECheckBoxState
+					{
+						if(OnGetToggleChecked.IsBound())
+						{
+							return OnGetToggleChecked.Execute(InComponent, ESlateRotationRepresentation::EulerXYZ, ESlateTransformSubComponent::Z);
+						}
+						return ECheckBoxState::Checked;
+					})
+					.OnToggleXChanged_Lambda([OnToggleChanged, InComponent](ECheckBoxState InState)
+					{
+						if(OnToggleChanged.IsBound())
+						{
+							OnToggleChanged.Execute(InComponent, ESlateRotationRepresentation::EulerXYZ, ESlateTransformSubComponent::X, InState);
+						}
+					})
+					.OnToggleYChanged_Lambda([OnToggleChanged, InComponent](ECheckBoxState InState)
+					{
+						if(OnToggleChanged.IsBound())
+						{
+							OnToggleChanged.Execute(InComponent, ESlateRotationRepresentation::EulerXYZ, ESlateTransformSubComponent::Y, InState);
+						}
+					})
+					.OnToggleZChanged_Lambda([OnToggleChanged, InComponent](ECheckBoxState InState)
+					{
+						if(OnToggleChanged.IsBound())
+						{
+							OnToggleChanged.Execute(InComponent, ESlateRotationRepresentation::EulerXYZ, ESlateTransformSubComponent::Z, InState);
+						}
+					});
+					
+				break;
+			}
+			case ESlateTransformComponent::Rotation:
+			{
+				typename SAdvancedRotationInputBox<NumericType>::FOnGetNumericValue RotationGetNumericValue;
+				if(OnGetNumericValue.IsBound())
+				{
+					RotationGetNumericValue = SAdvancedRotationInputBox<NumericType>::FOnGetNumericValue::CreateLambda(
+						[OnGetNumericValue, InComponent](ESlateRotationRepresentation::Type InRepr, ESlateTransformSubComponent::Type InSubComponent) -> TOptional<NumericType>
+					{
+						return OnGetNumericValue.Execute(InComponent, InRepr, InSubComponent);
+					});
+				}
+
+				typename SAdvancedRotationInputBox<NumericType>::FOnNumericValueChanged RotationValueChanged;
+				if(OnNumericValueChanged.IsBound())
+				{
+					RotationValueChanged = SAdvancedRotationInputBox<NumericType>::FOnNumericValueChanged::CreateLambda(
+						[OnNumericValueChanged, InComponent](ESlateRotationRepresentation::Type InRepr, ESlateTransformSubComponent::Type InSubComponent, NumericType InValue)
+					{
+						OnNumericValueChanged.Execute(InComponent, InRepr, InSubComponent, InValue);
+					});
+				}
+
+				typename SAdvancedRotationInputBox<NumericType>::FOnNumericValueCommitted RotationValueCommitted;
+				if(OnNumericValueCommitted.IsBound())
+				{
+					RotationValueCommitted = SAdvancedRotationInputBox<NumericType>::FOnNumericValueCommitted::CreateLambda(
+						[OnNumericValueCommitted, InComponent](ESlateRotationRepresentation::Type InRepr, ESlateTransformSubComponent::Type InSubComponent, NumericType InValue, ETextCommit::Type InCommitType)
+					{
+						OnNumericValueCommitted.Execute(InComponent, InRepr, InSubComponent, InValue, InCommitType);
+					});
+				}
+
+				TAttribute<TOptional<FQuat>> QuaternionAttribute;
+				TAttribute<TOptional<FRotator>> RotatorAttribute;
+				if(InArgs._Transform.IsBound() || InArgs._Transform.IsSet())
+				{
+					if(InArgs._UseQuaternionForRotation)
+					{
+						QuaternionAttribute = TAttribute<TOptional<FQuat>>::CreateLambda(
+							[OnGetQuaternion]() -> TOptional<FQuat>
+							{
+								return OnGetQuaternion();
+							});
+					}
+					else
+					{
+						RotatorAttribute = TAttribute<TOptional<FRotator>>::CreateLambda(
+							[OnGetRotator]() -> TOptional<FRotator>
+							{
+								return OnGetRotator();
+							});
+					}
+				}
+
+				typename SNumericVectorInputBox3::FOnNumericValueChanged XChanged, YChanged, ZChanged;
+				if(OnNumericValueChanged.IsBound())
+				{
+					XChanged = SNumericVectorInputBox3::FOnNumericValueChanged::CreateLambda(
+						[OnNumericValueChanged, InComponent](NumericType InValue)
+						{
+							OnNumericValueChanged.Execute(InComponent, ESlateRotationRepresentation::EulerXYZ, ESlateTransformSubComponent::X, InValue);
+						}
+					); 
+					YChanged = SNumericVectorInputBox3::FOnNumericValueChanged::CreateLambda(
+						[OnNumericValueChanged, InComponent](NumericType InValue)
+						{
+							OnNumericValueChanged.Execute(InComponent, ESlateRotationRepresentation::EulerXYZ, ESlateTransformSubComponent::Y, InValue);
+						}
+					); 
+					ZChanged = SNumericVectorInputBox3::FOnNumericValueChanged::CreateLambda(
+						[OnNumericValueChanged, InComponent](NumericType InValue)
+						{
+							OnNumericValueChanged.Execute(InComponent, ESlateRotationRepresentation::EulerXYZ, ESlateTransformSubComponent::Z, InValue);
+						}
+					); 
+				}
+					
+				typename SAdvancedRotationInputBox<NumericType>::FOnRotatorChanged RotatorChanged;
+				typename SAdvancedRotationInputBox<NumericType>::FOnQuaternionChanged QuaternionChanged;
+				if(OnTransformChanged.IsBound())
+				{
+					if(InArgs._UseQuaternionForRotation)
+					{
+						QuaternionChanged = SAdvancedRotationInputBox<NumericType>::FOnQuaternionChanged::CreateLambda(
+							[OnTransformChanged, OnGetTransform](FQuat InQuaternion)
+							{
+								TransformType Xfo = OnGetTransform();
+								Xfo.SetRotation(InQuaternion);
+								OnTransformChanged.Execute(Xfo);
+							}
+						);
+					}
+					else
+					{
+						RotatorChanged = SAdvancedRotationInputBox<NumericType>::FOnRotatorChanged::CreateLambda(
+							[OnTransformChanged, OnGetTransform](FRotator InRotator)
+							{
+								TransformType Xfo = OnGetTransform();
+								Xfo = TransformType(InRotator, Xfo.GetLocation(), Xfo.GetScale3D());
+								OnTransformChanged.Execute(Xfo);
+							}
+						);
+					}
+				}
+
+				typename SAdvancedRotationInputBox<NumericType>::FOnRotatorCommitted RotatorCommitted;
+				typename SAdvancedRotationInputBox<NumericType>::FOnQuaternionCommitted QuaternionCommitted;
+				if(OnTransformCommitted.IsBound())
+				{
+					if(InArgs._UseQuaternionForRotation)
+					{
+						QuaternionCommitted = SAdvancedRotationInputBox<NumericType>::FOnQuaternionCommitted::CreateLambda(
+							[OnTransformCommitted, OnGetTransform](FQuat InQuaternion, ETextCommit::Type InCommitType)
+							{
+								TransformType Xfo = OnGetTransform();
+								Xfo.SetRotation(InQuaternion);
+								OnTransformCommitted.Execute(Xfo, InCommitType);
+							}
+						);
+					}
+					else
+					{
+						RotatorCommitted = SAdvancedRotationInputBox<NumericType>::FOnRotatorCommitted::CreateLambda(
+							[OnTransformCommitted, OnGetTransform](FRotator InRotator, ETextCommit::Type InCommitType)
+							{
+								TransformType Xfo = OnGetTransform();
+								Xfo = TransformType(InRotator, Xfo.GetLocation(), Xfo.GetScale3D());
+								OnTransformCommitted.Execute(Xfo, InCommitType);
+							}
+						);
+					}
+				}
+
+				TSharedPtr<ESlateRotationRepresentation::Type> RotationRepresentationPtr = InArgs._RotationRepresentation;
+
+				InputWidget = SNew(SAdvancedRotationInputBox<NumericType>)
+					.Font(InArgs._Font)
+					.AllowSpin(InArgs._AllowSpin)
+					.bColorAxisLabels(InArgs._bColorAxisLabels)
+					.Representation_Lambda([RotationRepresentationPtr]() -> ESlateRotationRepresentation::Type
+					{
+						if(RotationRepresentationPtr.IsValid())
+						{
+							return *RotationRepresentationPtr.Get();
+						}
+						return ESlateRotationRepresentation::Rotator;
+					})
+					.OnGetNumericValue(RotationGetNumericValue)
+					.OnNumericValueChanged(RotationValueChanged)
+					.OnNumericValueCommitted(RotationValueCommitted)
+					.Rotator(RotatorAttribute)
+					.OnRotatorChanged(RotatorChanged)
+					.OnRotatorCommitted(RotatorCommitted)
+					.Quaternion(QuaternionAttribute)
+					.OnQuaternionChanged(QuaternionChanged)
+					.OnQuaternionCommitted(QuaternionCommitted)
+					.DisplayToggle(InArgs._DisplayToggle)
+					.TogglePadding(InArgs._TogglePadding)
+					.OnGetToggleChecked(SAdvancedRotationInputBox<NumericType>::FOnGetToggleChecked::CreateLambda(
+						[OnGetToggleChecked, InComponent](ESlateRotationRepresentation::Type InRepr, ESlateTransformSubComponent::Type InSubComponent) -> ECheckBoxState
+					{
+						if(OnGetToggleChecked.IsBound())
+						{
+							return OnGetToggleChecked.Execute(InComponent, InRepr, InSubComponent);
+						}
+						return ECheckBoxState::Checked;
+					}))
+					.OnToggleChanged(SAdvancedRotationInputBox<NumericType>::FOnToggleChanged::CreateLambda(
+						[OnToggleChanged, InComponent](ESlateRotationRepresentation::Type InRepr, ESlateTransformSubComponent::Type InSubComponent, ECheckBoxState InState)
+					{
+						if(OnToggleChanged.IsBound())
+						{
+							return OnToggleChanged.Execute(InComponent, InRepr, InSubComponent, InState);
+						}
+					}));
+
+				break;
+			}
+		}
+
+		if(InputWidget != SNullWidget::NullWidget)
+		{
+			HorizontalBox->AddSlot()
+			.HAlign(HAlign_Fill)
+			.VAlign(VAlign_Center)
+			.FillWidth(1.f)
+			[
+				InputWidget
+			];
+		}
+
+		return HorizontalBox;
+	}
+
+	/**
+	 * Constructs a label widget
+	 */
+	static TSharedRef<SWidget> ConstructLabel(const FArguments& InArgs, ESlateTransformComponent::Type InComponent)
+	{
+		TSharedRef<SWidget> LabelWidget = SNullWidget::NullWidget;
+
+		if(InComponent == ESlateTransformComponent::Rotation && InArgs._AllowEditRotationRepresentation)
+		{
+			const TArray<FText>& Labels = GetRotationRepresentationLabels();
+
+			TSharedPtr<ESlateRotationRepresentation::Type> RotationRepresentationPtr = InArgs._RotationRepresentation;
+			FOnRotationRepresentationChanged OnRotationRepresentationChanged = InArgs._OnRotationRepresentationChanged;
+
+			LabelWidget = SNew(SComboButton)
+				.ContentPadding(0)
+				.OnGetMenuContent_Lambda([Labels, RotationRepresentationPtr, OnRotationRepresentationChanged]
+				{
+					FMenuBuilder MenuBuilder(true, nullptr);
+					for (int32 LabelIndex = 0; LabelIndex < Labels.Num(); LabelIndex++)
+					{
+						MenuBuilder.AddMenuEntry(
+							Labels[LabelIndex],
+							FText(),
+							FSlateIcon(),
+							FUIAction
+							(
+								FExecuteAction::CreateLambda([LabelIndex, RotationRepresentationPtr, OnRotationRepresentationChanged]()
+								{
+									if(OnRotationRepresentationChanged.IsBound())
+									{
+										OnRotationRepresentationChanged.Execute((ESlateRotationRepresentation::Type)LabelIndex);
+									}
+									if(RotationRepresentationPtr.IsValid())
+									{
+										*RotationRepresentationPtr.Get() = (ESlateRotationRepresentation::Type)LabelIndex;
+									}
+								}),
+								FCanExecuteAction(),
+								FIsActionChecked::CreateLambda([LabelIndex, RotationRepresentationPtr]() -> bool
+								{
+									if(RotationRepresentationPtr.IsValid())
+									{
+										return LabelIndex == (int32)*RotationRepresentationPtr.Get(); 
+									}
+									return false;
+								})
+							),
+							NAME_None,
+							EUserInterfaceActionType::Check);
+					}
+					return MenuBuilder.MakeWidget();
+				})
+				.ButtonContent()
+				[
+					SNew(STextBlock)
+					.Font(InArgs._Font)
+					.Text_Lambda([RotationRepresentationPtr, Labels]() -> FText
+					{
+						if(RotationRepresentationPtr.IsValid())
+						{
+							return Labels[(int32)*RotationRepresentationPtr.Get()];
+						}
+						return Labels[0];
+					})
+				];
+		}
+		else
+		{
+			const FText& LabelText =
+				(InComponent == ESlateTransformComponent::Location) ?
+					InArgs._LocationLabel :
+					InArgs._ScaleLabel;
+
+			LabelWidget = SNew(STextBlock)
+			.Font(InArgs._Font)
+			.Text(LabelText);
+		}
+
+		TSharedPtr<SHorizontalBox> HorizontalBox;
+		SAssignNew(HorizontalBox, SHorizontalBox);
+
+		HorizontalBox->AddSlot()
+		.HAlign(HAlign_Left)
+		.VAlign(VAlign_Center)
+		.AutoWidth()
+		[
+			LabelWidget
+		];
+
+		const TSharedRef<SWidget> ScaleLockWidget = ConstructScaleLockWidget(InArgs, InComponent);
+		if(ScaleLockWidget != SNullWidget::NullWidget)
+		{
+			HorizontalBox->AddSlot()
+			.HAlign(HAlign_Left)
+			.VAlign(VAlign_Center)
+			.Padding(FMargin(4.f, 0.f, 0.f, 0.f))
+			[
+				ScaleLockWidget
+			];
+		}
+
+		const TSharedRef<SWidget> RelativeWorldWidget = ConstructRelativeWorldWidget(InArgs, InComponent);
+		if(RelativeWorldWidget != SNullWidget::NullWidget)
+		{
+			HorizontalBox->AddSlot()
+			.HAlign(HAlign_Right)
+			.VAlign(VAlign_Center)
+			.Padding(FMargin(4.f, 0.f, 0.f, 0.f))
+			[
+				RelativeWorldWidget
+			];
+		}
+
+		return HorizontalBox.ToSharedRef();
+	}
+
+	/**
+	 * Constructs a relative / world switcher button
+	 */
+	static TSharedRef<SWidget> ConstructScaleLockWidget(const FArguments& InArgs, ESlateTransformComponent::Type InComponent )
+	{
+		if(InArgs._DisplayScaleLock && InArgs._ConstructScale && InComponent == ESlateTransformComponent::Scale && InArgs._IsScaleLocked.IsValid())
+		{
+			TSharedPtr<bool> LockState = InArgs._IsScaleLocked;
+			FOnScaleLockChanged OnScaleLockChanged = InArgs._OnScaleLockChanged;
+
+			static const FText TooltipText = NSLOCTEXT("SAdvancedTransformInputBox", "PreserveScaleToolTip",
+					 "When locked, scales uniformly based on the current xyz scale values so the object maintains its shape in each direction when scaled");
+
+			return SNew(SButton)
+				.ContentPadding(0)
+				.ButtonStyle(FAppStyle::Get(), "SimpleButton")
+				.OnClicked_Lambda([OnScaleLockChanged, LockState]() -> FReply
+				{
+					*LockState.Get() = !*LockState.Get();
+					if(OnScaleLockChanged.IsBound())
+					{
+						OnScaleLockChanged.Execute(*LockState.Get());
+					}					
+					return FReply::Handled();
+				})
+				.ToolTipText(TooltipText)
+				.HAlign(HAlign_Center)
+				.VAlign(VAlign_Center)
+				.ForegroundColor(FSlateColor::UseStyle())
+				.Content()
+				[
+					SNew(SImage)
+					.Image_Lambda([LockState]() -> const FSlateBrush*
+					{
+						return GetScaleLockIcon(*LockState.Get());
+					})
+					.ColorAndOpacity(FSlateColor::UseForeground())
+				];
+		}
+		return SNullWidget::NullWidget;
+	}
+
+	/**
+	 * Gather the icon for scale locked / unlocked
+	 */
+	static const FSlateBrush* GetScaleLockIcon(bool bIsLocked)
+	{
+		return bIsLocked ? FAppStyle::Get().GetBrush( TEXT("Icons.Lock") ) : FAppStyle::Get().GetBrush( TEXT("Icons.Unlock") ) ;
+	}
+	
+	/**
+	 * Constructs a relative / world switcher button
+	 */
+	static TSharedRef<SWidget> ConstructRelativeWorldWidget(const FArguments& InArgs, ESlateTransformComponent::Type InComponent)
+	{
+		if(InArgs._DisplayRelativeWorld)
+		{
+			FOnGetIsComponentRelative OnGetIsComponentRelative = InArgs._OnGetIsComponentRelative;;
+			FOnIsComponentRelativeChanged OnIsComponentRelativeChanged = InArgs._OnIsComponentRelativeChanged;
+
+			static const FText TooltipFormat = NSLOCTEXT("SAdvancedTransformInputBox", "RelativeWorldTooltip",
+					"Cycles the Transform {0} coordinate system between world and local (relative) space.\nShift clicking this button affects the whole transform.");
+			static const FText LocationTooltip = FText::Format(TooltipFormat, NSLOCTEXT("SAdvancedTransformInputBox", "Location", "Location"));
+			static const FText RotationTooltip = FText::Format(TooltipFormat, NSLOCTEXT("SAdvancedTransformInputBox", "Rotation", "Rotation"));
+			static const FText ScaleTooltip = FText::Format(TooltipFormat, NSLOCTEXT("SAdvancedTransformInputBox", "Scale", "Scale"));
+
+			FText TooltipText;
+			switch(InComponent)
+			{
+				case ESlateTransformComponent::Location:
+				{
+					TooltipText = LocationTooltip;
+					break;
+				}
+				case ESlateTransformComponent::Rotation:
+				{
+					TooltipText = RotationTooltip;
+					break;
+				}
+				case ESlateTransformComponent::Scale:
+				{
+					TooltipText = ScaleTooltip;
+					break;
+				}
+			}
+
+			return SNew(SButton)
+				.ContentPadding(0)
+				.ButtonStyle(FAppStyle::Get(), "SimpleButton")
+				.OnClicked_Lambda([OnIsComponentRelativeChanged, OnGetIsComponentRelative, InComponent]() -> FReply
+				{
+					if(OnIsComponentRelativeChanged.IsBound())
+					{
+						bool bIsRelative = true;
+						if(OnGetIsComponentRelative.IsBound())
+						{
+							bIsRelative = OnGetIsComponentRelative.Execute(InComponent);
+						}
+
+						if(FSlateApplication::Get().GetModifierKeys().IsShiftDown() ||
+							FSlateApplication::Get().GetModifierKeys().IsControlDown() ||
+							FSlateApplication::Get().GetModifierKeys().IsAltDown())
+						{
+							OnIsComponentRelativeChanged.Execute(ESlateTransformComponent::Location, !bIsRelative);
+							OnIsComponentRelativeChanged.Execute(ESlateTransformComponent::Rotation, !bIsRelative);
+							OnIsComponentRelativeChanged.Execute(ESlateTransformComponent::Scale, !bIsRelative);
+						}
+						else
+						{
+							OnIsComponentRelativeChanged.Execute(InComponent, !bIsRelative);
+						}
+					}					
+					return FReply::Handled();
+				})
+				.ToolTipText(TooltipText)
+				.HAlign(HAlign_Center)
+				.VAlign(VAlign_Center)
+				.ForegroundColor(FSlateColor::UseStyle())
+				.Content()
+				[
+					SNew(SImage)
+					.Image_Lambda([OnGetIsComponentRelative, InComponent]() -> const FSlateBrush*
+					{
+						bool bIsRelative = true;
+						if(OnGetIsComponentRelative.IsBound())
+						{
+							bIsRelative = OnGetIsComponentRelative.Execute(InComponent);
+						}
+						return GetRelativeWorldIcon(bIsRelative).GetIcon();
+					})
+					.ColorAndOpacity(FSlateColor::UseForeground())
+				];
+		}
+		return SNullWidget::NullWidget;
+	}
+
+	/**
+	 * Gather the icon for relative world
+	 */
+	static const FSlateIcon& GetRelativeWorldIcon(bool bRelative)
+	{
+		if (bRelative)
+		{
+			static FName LocalIconName("Icons.Transform");
+			static FSlateIcon LocalIcon = FSlateIcon(FAppStyle::GetAppStyleSetName(), LocalIconName);
+			return LocalIcon;
+		}
+		
+		static FName WorldIconName("EditorViewport.RelativeCoordinateSystem_World");
+		static FSlateIcon WorldIcon = FSlateIcon(FAppStyle::GetAppStyleSetName(), WorldIconName);
+		return WorldIcon;
+	}
+
+	static const TArray<FText>& GetRotationRepresentationLabels()
+	{
+		static TArray<FText> Labels;
+		if(Labels.IsEmpty())
+		{
+			Labels.Add(NSLOCTEXT("SAdvancedTransformInputBox", "EulerXYZ", "Euler XYZ"));
+			Labels.Add(NSLOCTEXT("SAdvancedTransformInputBox", "EulerXZY", "Euler XZY"));
+			Labels.Add(NSLOCTEXT("SAdvancedTransformInputBox", "EulerYXZ", "Euler YXZ"));
+			Labels.Add(NSLOCTEXT("SAdvancedTransformInputBox", "EulerYZX", "Euler YZX"));
+			Labels.Add(NSLOCTEXT("SAdvancedTransformInputBox", "EulerZXY", "Euler ZXY"));
+			Labels.Add(NSLOCTEXT("SAdvancedTransformInputBox", "EulerZYX", "Euler ZYX"));
+			Labels.Add(NSLOCTEXT("SAdvancedTransformInputBox", "Rotator", "Rotator"));
+			Labels.Add(NSLOCTEXT("SAdvancedTransformInputBox", "Quaternion", "Quaternion"));
+			Labels.Add(NSLOCTEXT("SAdvancedTransformInputBox", "AxisAndAngle", "Axis And Angle"));
+		}
+		return Labels;
+	}
+
+	static void ConstrainScale(
+		int32 ComponentIndex,
+		UE::Math::TVector<NumericType> OldValue,
+		UE::Math::TVector<NumericType>& NewValue,
+		TAttribute<bool> ConstrainComponents)
+	{
+		if(!ConstrainComponents.Get())
+		{
+			return;
+		}
+
+		if(FMath::IsNearlyZero(OldValue[ComponentIndex]) || FMath::IsNearlyZero(NewValue[ComponentIndex]))
+		{
+			return;
+		}
+
+		const NumericType Ratio = NewValue[ComponentIndex] / OldValue[ComponentIndex];
+		NewValue = OldValue * Ratio;
+	}
+};
diff --git a/Engine/Source/Runtime/AnimationCore/Public/EulerTransform.h b/Engine/Source/Runtime/AnimationCore/Public/EulerTransform.h
--- a/Engine/Source/Runtime/AnimationCore/Public/EulerTransform.h
+++ b/Engine/Source/Runtime/AnimationCore/Public/EulerTransform.h
@@ -40,6 +40,13 @@ struct ANIMATIONCORE_API FEulerTransform
 	{
 	}
 
+	FORCEINLINE FEulerTransform(const FRotator& InRotation, const FVector& InLocation, const FVector& InScale)
+		: Location(InLocation)
+		, Rotation(InRotation)
+		, Scale(InScale)
+	{
+	}
+
 	FORCEINLINE FEulerTransform(const FTransform& InTransform)
 		: Location(InTransform.GetLocation())
 		, Rotation(InTransform.GetRotation().Rotator())
@@ -84,4 +91,14 @@ struct ANIMATIONCORE_API FEulerTransform
 		Rotation = InTransform.GetRotation().Rotator();
 		Scale = InTransform.GetScale3D();
 	}
+
+	FORCEINLINE const FVector& GetLocation() const { return Location; }
+	FORCEINLINE FQuat GetRotation() const { return Rotation.Quaternion(); }
+	FORCEINLINE const FRotator& Rotator() const { return Rotation; }
+	FORCEINLINE const FVector& GetScale3D() const { return Scale; }
+	FORCEINLINE void SetLocation(const FVector& InValue) { Location = InValue; }
+	FORCEINLINE void SetRotation(const FQuat& InValue) { Rotation = InValue.Rotator(); }
+	FORCEINLINE void SetRotator(const FRotator& InValue) { Rotation = InValue; }
+	FORCEINLINE void SetScale3D(const FVector& InValue) { Scale = InValue; }
+	FORCEINLINE void NormalizeRotation() {}
 };
\ No newline at end of file
diff --git a/Engine/Source/Runtime/Slate/Public/Widgets/Input/SSegmentedControl.h b/Engine/Source/Runtime/Slate/Public/Widgets/Input/SSegmentedControl.h
--- a/Engine/Source/Runtime/Slate/Public/Widgets/Input/SSegmentedControl.h
+++ b/Engine/Source/Runtime/Slate/Public/Widgets/Input/SSegmentedControl.h
@@ -116,10 +116,14 @@ public:
 	}
 
 	DECLARE_DELEGATE_OneParam( FOnValueChanged, OptionType );
+	DECLARE_DELEGATE_OneParam( FOnValuesChanged, TArray<OptionType> );
+	DECLARE_DELEGATE_TwoParams( FOnValueChecked, OptionType, ECheckBoxState );
 
 	SLATE_BEGIN_ARGS( SSegmentedControl<OptionType> )
 		: _Style(&FAppStyle::Get().GetWidgetStyle<FSegmentedControlStyle>("SegmentedControl"))
 		, _TextStyle(&FAppStyle::Get().GetWidgetStyle<FTextBlockStyle>("SmallButtonText"))
+		, _SupportsMultiSelection(false)
+		, _SupportsEmptySelection(false)
 		, _MaxSegmentsPerLine(0)
 	{}
 		/** Slot type supported by this panel */
@@ -131,22 +135,44 @@ public:
 		/** Styling for the text in each slot. If a custom widget is supplied for a slot this argument is not used */
 		SLATE_STYLE_ARGUMENT(FTextBlockStyle, TextStyle)
 
-		/** The current control value */
+		/**
+		 * If enabled the widget will support multi selection.
+		 * For single selection the widget relies on the Value attribute,
+		 * for multi selection the widget relies on the MultiValue attribute.
+		 */
+		SLATE_ARGUMENT(bool, SupportsMultiSelection)
+
+		/**
+		 * If enabled the widget will support an empty selection.
+		 * This is only enabled if SupportsMultiSelection is also enabled.
+		 */
+		SLATE_ARGUMENT(bool, SupportsEmptySelection)
+
+		/** The current control value. */
 		SLATE_ATTRIBUTE(OptionType, Value)
 
+		/** The current (multiple) control values (if SupportsMultiSelection is enabled) */ 
+		SLATE_ATTRIBUTE(TArray<OptionType>, Values)
+	
 		/** Padding to apply to each slot */
 		SLATE_ATTRIBUTE(FMargin, UniformPadding)
 
-		/** Called when the value is changed */
+		/** Called when the (primary) value is changed */
 		SLATE_EVENT(FOnValueChanged, OnValueChanged)
-		
+
+		/** Called when the any value is changed */
+		SLATE_EVENT(FOnValuesChanged, OnValuesChanged)
+
+		/** Called when the value is changed (useful for multi selection) */
+		SLATE_EVENT(FOnValueChecked, OnValueChecked)
+
 		/** Optional maximum number of segments per line before the control wraps vertically to the next line. If this value is <= 0 no wrapping happens */
 		SLATE_ARGUMENT(int32, MaxSegmentsPerLine)
 	SLATE_END_ARGS()
 
 	SSegmentedControl()
 		: Children(this)
-		, CurrentValue(*this)
+		, CurrentValues(*this)
 	{}
 
 	void Construct( const FArguments& InArgs )
@@ -156,9 +182,22 @@ public:
 		Style = InArgs._Style;
 		TextStyle = InArgs._TextStyle;
 
-		CurrentValueIsBound = InArgs._Value.IsBound();
-		CurrentValue.Assign(*this, InArgs._Value);
+		SupportsMultiSelection = InArgs._SupportsMultiSelection;
+		SupportsEmptySelection = InArgs._SupportsEmptySelection;
+		CurrentValuesIsBound = false; // will be set by SetValue or SetValues
+
+		if(InArgs._Value.IsBound() || InArgs._Value.IsSet())
+		{
+			SetValue(InArgs._Value, false);
+		}
+		else if(InArgs._Values.IsBound() || InArgs._Values.IsSet())
+		{
+			SetValues(InArgs._Values, false);
+		}
+
 		OnValueChanged = InArgs._OnValueChanged;
+		OnValuesChanged = InArgs._OnValuesChanged;
+		OnValueChecked = InArgs._OnValueChecked;
 
 		UniformPadding = InArgs._UniformPadding;
 
@@ -263,6 +302,7 @@ public:
 			]
 		];
 
+		UpdateCheckboxValuesIfNeeded();
 	}
 
 	// Slot Management
@@ -293,37 +333,83 @@ public:
 
 	OptionType GetValue() const
 	{
-		return CurrentValue.Get();
+		const TArray<OptionType> Values = GetValues();
+		if(Values.IsEmpty())
+		{
+			return OptionType();
+		}
+		return Values[0];
 	}
 
-	/** See the Value attribute */
-	void SetValue(TAttribute<OptionType> InValue) 
+	TArray<OptionType> GetValues() const
 	{
-		CurrentValueIsBound = InValue.IsBound();
-		CurrentValue.Assign(*this, MoveTemp(InValue));
+		return CurrentValues.Get();
+	}
 
-		if (!CurrentValueIsBound)
+	bool HasValue(OptionType InValue)
+	{
+		const TArray<OptionType> Values = GetValues();
+		return Values.Contains(InValue);
+	}
+
+	/** See the Value attribute */
+	void SetValue(TAttribute<OptionType> InValue, bool bUpdateChildren = true)
+	{
+		if(InValue.IsBound())
 		{
-			for (int32 Index = 0; Index < Children.Num(); ++Index)
+			SetValues(TAttribute<TArray<OptionType>>::CreateLambda([InValue]() -> TArray<OptionType>
 			{
-				const FSlot& Slot = Children[Index];
-				if (TSharedPtr<SCheckBox> CheckBox = Slot._CheckBox.Pin())
-				{
-					CheckBox->SetIsChecked(Slot._Value == CurrentValue.Get() ? ECheckBoxState::Checked : ECheckBoxState::Unchecked);
-				}
-			}
+				TArray<OptionType> Values;
+				Values.Add(InValue.Get());
+				return Values;
+			}), bUpdateChildren); 
+		}
+		else if(InValue.IsSet())
+		{
+			TArray<OptionType> Values = {InValue.Get()};
+			SetValues(TAttribute<TArray<OptionType>>(Values), bUpdateChildren); 
+		}
+		else
+		{
+			SetValues(TAttribute<TArray<OptionType>>(), bUpdateChildren);
+		}
+	}
+
+	/** See the Values attribute */
+	void SetValues(TAttribute<TArray<OptionType>> InValues, bool bUpdateChildren = true) 
+	{
+		CurrentValuesIsBound = InValues.IsBound();
+
+		if(CurrentValuesIsBound)
+		{
+			CurrentValues.Assign(*this, InValues);
+		}
+		else if(InValues.IsSet())
+		{
+			CurrentValues.Set(*this, InValues.Get());
+		}
+		else
+		{
+			CurrentValues.Unbind(*this);
+		};
+
+		if(bUpdateChildren)
+		{
+			UpdateCheckboxValuesIfNeeded();
 		}
 	}
 
 private:
+	
 	TAttribute<ECheckBoxState> GetCheckBoxStateAttribute(OptionType InValue) const
 	{
 		auto Lambda = [this, InValue]()
 		{
-			return InValue == CurrentValue.Get() ? ECheckBoxState::Checked : ECheckBoxState::Unchecked;
+			const TArray<OptionType> Values = GetValues();
+			return Values.Contains(InValue) ? ECheckBoxState::Checked : ECheckBoxState::Unchecked;
 		};
 
-		if (CurrentValueIsBound)
+		if (CurrentValuesIsBound)
 		{
 			return MakeAttributeLambda(Lambda);
 		}
@@ -331,34 +417,106 @@ private:
 		return Lambda();
 	}
 
+	void UpdateCheckboxValuesIfNeeded()
+	{
+		if (!CurrentValuesIsBound)
+		{
+			const TArray<OptionType> Values = GetValues();
+			
+			for (int32 Index = 0; Index < Children.Num(); ++Index)
+			{
+				const FSlot& Slot = Children[Index];
+				if (const TSharedPtr<SCheckBox> CheckBox = Slot._CheckBox.Pin())
+				{
+					CheckBox->SetIsChecked(Values.Contains(Slot._Value) ? ECheckBoxState::Checked : ECheckBoxState::Unchecked);
+				}
+			}
+		}
+	}
+
 	void CommitValue(const ECheckBoxState InCheckState, OptionType InValue)
 	{
-		if (InCheckState == ECheckBoxState::Checked)
+		const TArray<OptionType> PreviousValues = CurrentValues.Get();
+		TArray<OptionType> Values = PreviousValues;
+
+		// don't allow to deselect the last checkbox
+		if(InCheckState != ECheckBoxState::Checked && Values.Num() == 1)
+		{
+			if(!SupportsEmptySelection)
+			{
+				UpdateCheckboxValuesIfNeeded();
+				return;
+			}
+		}
+
+		bool bModifierIsDown = false;
+		if(SupportsMultiSelection)
+		{
+			bModifierIsDown =
+				FSlateApplication::Get().GetModifierKeys().IsShiftDown() ||
+				FSlateApplication::Get().GetModifierKeys().IsControlDown();
+		}
+
+		// if the attribute is not bound update our internal state
+		if (!CurrentValuesIsBound)
 		{
-			// don't overwrite the bound attribute, but still notify that the value was committed.
-			if (!CurrentValueIsBound)
+			if(bModifierIsDown)
 			{
-				CurrentValue.Set(*this, InValue);
-				for (int32 Index = 0; Index < Children.Num(); ++Index)
+				if (InCheckState == ECheckBoxState::Checked)
 				{
-					const FSlot& Slot = Children[Index];
-					if (TSharedPtr<SCheckBox> CheckBox = Slot._CheckBox.Pin())
+					Values.AddUnique(InValue);
+				}
+				else
+				{
+					Values.Remove(InValue);
+				}
+			}
+			else
+			{
+				if(InCheckState == ECheckBoxState::Checked)
+				{
+					Values.Reset();
+					Values.Add(InValue);
+				}
+			}
+
+			CurrentValues.Set(*this, Values);
+
+			UpdateCheckboxValuesIfNeeded();
+		}
+
+		if(OnValueChecked.IsBound())
+		{
+			if(!bModifierIsDown && InCheckState == ECheckBoxState::Checked)
+			{
+				for(const OptionType PreviousValue : PreviousValues)
+				{
+					if(!Values.Contains(PreviousValue))
 					{
-						CheckBox->SetIsChecked(Slot._Value == CurrentValue.Get() ? ECheckBoxState::Checked : ECheckBoxState::Unchecked);
+						OnValueChecked.Execute(PreviousValue, ECheckBoxState::Unchecked);
 					}
-			    }
+				}
 			}
 
-			OnValueChanged.ExecuteIfBound( InValue );
+			OnValueChecked.Execute(InValue, InCheckState);
 		}
+
+		if (InCheckState == ECheckBoxState::Checked)
+		{
+			OnValueChanged.ExecuteIfBound(InValue);
+		}
+
+		OnValuesChanged.ExecuteIfBound(Values);
 	}
 
 private:
 	TPanelChildren<FSlot> Children;
 
 	FOnValueChanged OnValueChanged;
+	FOnValuesChanged OnValuesChanged;
+	FOnValueChecked OnValueChecked;
 
-	TSlateAttribute<OptionType, EInvalidateWidgetReason::Paint> CurrentValue;
+	TSlateAttribute<TArray<OptionType>, EInvalidateWidgetReason::Paint> CurrentValues;
 
 	TAttribute<FMargin> UniformPadding;
 
@@ -368,5 +526,7 @@ private:
 
 	int32 MaxSegmentsPerLine = 0;
 
-	bool CurrentValueIsBound = false;
+	bool CurrentValuesIsBound = false;
+	bool SupportsMultiSelection = false;
+	bool SupportsEmptySelection = false;
 };
