Improve SAdvancedTransformInputBox

diff --git a/Engine/Source/Developer/AnimationWidgets/Public/SAdvancedTransformInputBox.h b/Engine/Source/Developer/AnimationWidgets/Public/SAdvancedTransformInputBox.h
--- a/Engine/Source/Developer/AnimationWidgets/Public/SAdvancedTransformInputBox.h
+++ b/Engine/Source/Developer/AnimationWidgets/Public/SAdvancedTransformInputBox.h
@@ -75,7 +75,8 @@ public:
 		, _LabelPadding( FMargin( 0.f,0.f,6.f,0.f ) )
 		, _Font(FAppStyle::Get().GetFontStyle("NormalFont"))
 		, _AllowSpin(true)
-		, _SpinDelta(0.001f)
+		, _LocationSpinDelta(0.01f)
+		, _ScaleSpinDelta(0.001f)
 		, _AllowEditRotationRepresentation(true)
 		, _DisplayScaleLock(true)
 		, _IsScaleLocked(TSharedPtr<bool>(new bool(false)))
@@ -144,8 +145,11 @@ public:
 		/** Whether or not values can be spun or if they should be typed in */
 		SLATE_ARGUMENT( bool, AllowSpin )
 
-		/** The delta amount to apply, per pixel, when the spinner is dragged. */
-		SLATE_ATTRIBUTE( NumericType, SpinDelta )
+		/** The delta amount to apply, per pixel, when a location spinner is dragged. */
+		SLATE_ATTRIBUTE( NumericType, LocationSpinDelta )
+
+		/** The delta amount to apply, per pixel, when a scale spinner is dragged. */
+		SLATE_ATTRIBUTE( NumericType, ScaleSpinDelta )
 
 		/** Whether or not to display the rotation representation picker in the label */
 		SLATE_ARGUMENT( bool, AllowEditRotationRepresentation )
@@ -224,7 +228,9 @@ public:
 	 */
 	static TSharedRef<SWidget> ConstructWidget(const FArguments& InArgs, ESlateTransformComponent::Type InComponent)
 	{
-		TSharedRef<SHorizontalBox> HorizontalBox = SNew(SHorizontalBox);
+		SHorizontalBox::FArguments BoxArgs;
+		((FSlateBaseNamedArgs&)BoxArgs) = (FSlateBaseNamedArgs)InArgs;
+		TSharedRef<SHorizontalBox> HorizontalBox = SArgumentNew(BoxArgs, SHorizontalBox);
 
 		if(InArgs._ShowInlineLabels)
 		{
@@ -505,7 +511,7 @@ public:
 				InputWidget = SNew(SNumericVectorInputBox3)
 					.Font(InArgs._Font)
 					.AllowSpin(InArgs._AllowSpin)
-					.SpinDelta(InArgs._SpinDelta)
+					.SpinDelta(InComponent == ESlateTransformComponent::Location ? InArgs._LocationSpinDelta : InArgs._ScaleSpinDelta)
 					.bColorAxisLabels(InArgs._bColorAxisLabels)
 					.X(XAttribute)
 					.Y(YAttribute)
@@ -834,17 +840,20 @@ public:
 		else
 		{
 			const FText& LabelText =
+				(InComponent == ESlateTransformComponent::Rotation) ?
+					InArgs._RotationLabel : (
 				(InComponent == ESlateTransformComponent::Location) ?
 					InArgs._LocationLabel :
-					InArgs._ScaleLabel;
+					InArgs._ScaleLabel);
 
 			LabelWidget = SNew(STextBlock)
 			.Font(InArgs._Font)
 			.Text(LabelText);
 		}
 
-		TSharedPtr<SHorizontalBox> HorizontalBox;
-		SAssignNew(HorizontalBox, SHorizontalBox);
+		SHorizontalBox::FArguments BoxArgs;
+		((FSlateBaseNamedArgs&)BoxArgs) = (FSlateBaseNamedArgs)InArgs;
+		TSharedRef<SHorizontalBox> HorizontalBox = SArgumentNew(BoxArgs, SHorizontalBox);
 
 		HorizontalBox->AddSlot()
 		.HAlign(HAlign_Left)
@@ -878,7 +887,7 @@ public:
 			];
 		}
 
-		return HorizontalBox.ToSharedRef();
+		return HorizontalBox;
 	}
 
 	/**
