SAdvancedTransformWidget: Add helper functions for numeric value getters / setters

diff --git a/Engine/Source/Developer/AnimationWidgets/Public/SAdvancedTransformInputBox.h b/Engine/Source/Developer/AnimationWidgets/Public/SAdvancedTransformInputBox.h
--- a/Engine/Source/Developer/AnimationWidgets/Public/SAdvancedTransformInputBox.h
+++ b/Engine/Source/Developer/AnimationWidgets/Public/SAdvancedTransformInputBox.h
@@ -775,6 +775,8 @@ public:
 	 */
 	static TSharedRef<SWidget> ConstructLabel(const FArguments& InArgs, ESlateTransformComponent::Type InComponent)
 	{
+		const TAttribute<bool> EnabledAttribute = InArgs._IsEnabled;
+
 		TSharedRef<SWidget> LabelWidget = SNullWidget::NullWidget;
 
 		if(InComponent == ESlateTransformComponent::Rotation && InArgs._AllowEditRotationRepresentation)
@@ -785,6 +787,7 @@ public:
 			FOnRotationRepresentationChanged OnRotationRepresentationChanged = InArgs._OnRotationRepresentationChanged;
 
 			LabelWidget = SNew(SComboButton)
+				.IsEnabled(true)
 				.ContentPadding(0)
 				.OnGetMenuContent_Lambda([Labels, RotationRepresentationPtr, OnRotationRepresentationChanged]
 				{
@@ -847,12 +850,14 @@ public:
 					InArgs._ScaleLabel);
 
 			LabelWidget = SNew(STextBlock)
+			.IsEnabled(EnabledAttribute)
 			.Font(InArgs._Font)
 			.Text(LabelText);
 		}
 
 		SHorizontalBox::FArguments BoxArgs;
 		((FSlateBaseNamedArgs&)BoxArgs) = (FSlateBaseNamedArgs)InArgs;
+		BoxArgs.IsEnabled(true);
 		TSharedRef<SHorizontalBox> HorizontalBox = SArgumentNew(BoxArgs, SHorizontalBox);
 
 		HorizontalBox->AddSlot()
@@ -866,6 +871,8 @@ public:
 		const TSharedRef<SWidget> ScaleLockWidget = ConstructScaleLockWidget(InArgs, InComponent);
 		if(ScaleLockWidget != SNullWidget::NullWidget)
 		{
+			ScaleLockWidget->SetEnabled(EnabledAttribute);
+			
 			HorizontalBox->AddSlot()
 			.HAlign(HAlign_Left)
 			.VAlign(VAlign_Center)
@@ -878,6 +885,9 @@ public:
 		const TSharedRef<SWidget> RelativeWorldWidget = ConstructRelativeWorldWidget(InArgs, InComponent);
 		if(RelativeWorldWidget != SNullWidget::NullWidget)
 		{
+			// the relative to world will be always enabled...
+			// (unless the outer scope if not enabled)
+			
 			HorizontalBox->AddSlot()
 			.HAlign(HAlign_Right)
 			.VAlign(VAlign_Center)
@@ -977,6 +987,7 @@ public:
 			}
 
 			return SNew(SButton)
+				.IsEnabled(true)
 				.ContentPadding(0)
 				.ButtonStyle(FAppStyle::Get(), "SimpleButton")
 				.OnClicked_Lambda([OnIsComponentRelativeChanged, OnGetIsComponentRelative, InComponent]() -> FReply
@@ -1080,4 +1091,305 @@ public:
 		const NumericType Ratio = NewValue[ComponentIndex] / OldValue[ComponentIndex];
 		NewValue = OldValue * Ratio;
 	}
+
+	static TOptional<NumericType> GetNumericValueFromTransform(
+		const TransformType& Transform,
+		ESlateTransformComponent::Type Component,
+		ESlateRotationRepresentation::Type Representation,
+		ESlateTransformSubComponent::Type SubComponent)
+	{
+		switch(Component)
+		{
+			case ESlateTransformComponent::Location:
+			{
+				switch(SubComponent)
+				{
+					case ESlateTransformSubComponent::X:
+						return Transform.GetLocation().X;
+					case ESlateTransformSubComponent::Y:
+						return Transform.GetLocation().Y;
+					case ESlateTransformSubComponent::Z:
+						return Transform.GetLocation().Z;
+				}
+				break;
+			}
+			case ESlateTransformComponent::Rotation:
+			{
+				switch(Representation)
+				{
+					case ESlateRotationRepresentation::EulerXYZ:
+					case ESlateRotationRepresentation::EulerXZY:
+					case ESlateRotationRepresentation::EulerYXZ:
+					case ESlateRotationRepresentation::EulerYZX:
+					case ESlateRotationRepresentation::EulerZXY:
+					case ESlateRotationRepresentation::EulerZYX:
+					{
+						const int32 RotationOrder = int32(Representation) - int32(ESlateRotationRepresentation::EulerXYZ);
+						const FVector Euler = AnimationCore::EulerFromQuat(Transform.GetRotation(), EEulerRotationOrder(RotationOrder));
+						switch(SubComponent)
+						{
+							case ESlateTransformSubComponent::X:
+							{
+								return Euler.X;
+							}
+							case ESlateTransformSubComponent::Y:
+							{
+								return Euler.Y;
+							}
+							case ESlateTransformSubComponent::Z:
+							{
+								return Euler.Z;
+							}
+						}
+						break;
+					}
+					case ESlateRotationRepresentation::Rotator:
+					{
+						switch(SubComponent)
+						{
+							case ESlateTransformSubComponent::Roll:
+								return Transform.Rotator().Roll;
+							case ESlateTransformSubComponent::Pitch:
+								return Transform.Rotator().Pitch;
+							case ESlateTransformSubComponent::Yaw:
+								return Transform.Rotator().Yaw;
+						}
+						break;
+					}
+					case ESlateRotationRepresentation::Quaternion:
+					{
+						switch(SubComponent)
+						{
+							case ESlateTransformSubComponent::X:
+								return Transform.GetRotation().X;
+							case ESlateTransformSubComponent::Y:
+								return Transform.GetRotation().Y;
+							case ESlateTransformSubComponent::Z:
+								return Transform.GetRotation().Z;
+							case ESlateTransformSubComponent::W:
+								return Transform.GetRotation().W;
+						}
+						break;
+					}
+					case ESlateRotationRepresentation::AxisAndAngle:
+					{
+						FVector Axis;
+						NumericType Angle;
+						Transform.GetRotation().ToAxisAndAngle(Axis, Angle);
+
+						switch(SubComponent)
+						{
+							case ESlateTransformSubComponent::X:
+								return Axis.X;
+							case ESlateTransformSubComponent::Y:
+								return Axis.Y;
+							case ESlateTransformSubComponent::Z:
+								return Axis.Z;
+							case ESlateTransformSubComponent::Angle:
+								return Angle;
+						}
+						break;
+					}
+				}
+				break;
+			}
+			case ESlateTransformComponent::Scale:
+			{
+				switch(SubComponent)
+				{
+					case ESlateTransformSubComponent::X:
+						return Transform.GetScale3D().X;
+					case ESlateTransformSubComponent::Y:
+						return Transform.GetScale3D().Y;
+					case ESlateTransformSubComponent::Z:
+						return Transform.GetScale3D().Z;
+				}
+				break;
+			}
+		}
+		return TOptional<NumericType>();
+	}
+
+	static void ApplyNumericValueChange(
+		TransformType& Transform,
+		NumericType Value,
+		ESlateTransformComponent::Type Component,
+		ESlateRotationRepresentation::Type Representation,
+		ESlateTransformSubComponent::Type SubComponent)
+	{
+		switch(Component)
+		{
+			case ESlateTransformComponent::Location:
+			{
+				UE::Math::TVector<NumericType> Location = Transform.GetLocation();
+				switch(SubComponent)
+				{
+					case ESlateTransformSubComponent::X:
+					{
+						Location.X = Value;
+						break;
+					}
+					case ESlateTransformSubComponent::Y:
+					{
+						Location.Y = Value;
+						break;
+					}
+					case ESlateTransformSubComponent::Z:
+					{
+						Location.Z = Value;
+						break;
+					}
+				}
+				Transform.SetLocation(Location);
+				break;
+			}
+			case ESlateTransformComponent::Rotation:
+			{
+				switch(Representation)
+				{
+					case ESlateRotationRepresentation::EulerXYZ:
+					case ESlateRotationRepresentation::EulerXZY:
+					case ESlateRotationRepresentation::EulerYXZ:
+					case ESlateRotationRepresentation::EulerYZX:
+					case ESlateRotationRepresentation::EulerZXY:
+					case ESlateRotationRepresentation::EulerZYX:
+					{
+						const int32 RotationOrder = int32(Representation) - int32(ESlateRotationRepresentation::EulerXYZ);
+						FVector Euler = AnimationCore::EulerFromQuat(Transform.GetRotation(), EEulerRotationOrder(RotationOrder));
+						switch(SubComponent)
+						{
+							case ESlateTransformSubComponent::X:
+							{
+								Euler.X = Value;
+								break;
+							}
+							case ESlateTransformSubComponent::Y:
+							{
+								Euler.Y = Value;
+								break;
+							}
+							case ESlateTransformSubComponent::Z:
+							{
+								Euler.Z = Value;
+								break;
+							}
+						}
+						Transform.SetRotation(AnimationCore::QuatFromEuler(Euler, EEulerRotationOrder(RotationOrder)));
+						break;
+					}
+					case ESlateRotationRepresentation::Rotator:
+					{
+						FRotator Rotator = Transform.Rotator();
+						switch(SubComponent)
+						{
+							case ESlateTransformSubComponent::Roll:
+							{
+								Rotator.Roll = Value;
+								break;
+							}
+							case ESlateTransformSubComponent::Pitch:
+							{
+								Rotator.Pitch = Value;
+								break;
+							}
+							case ESlateTransformSubComponent::Yaw:
+							{
+								Rotator.Yaw = Value;
+								break;
+							}
+						}
+						Transform = TransformType(Rotator, Transform.GetLocation(), Transform.GetScale3D());
+						break;
+					}
+					case ESlateRotationRepresentation::Quaternion:
+					{
+						FQuat Rotation = Transform.GetRotation();
+						switch(SubComponent)
+						{
+							case ESlateTransformSubComponent::X:
+							{
+								Rotation.X = Value;
+								break;
+							}
+							case ESlateTransformSubComponent::Y:
+							{
+								Rotation.Y = Value;
+								break;
+							}
+							case ESlateTransformSubComponent::Z:
+							{
+								Rotation.Z = Value;
+								break;
+							}
+							case ESlateTransformSubComponent::W:
+							{
+								Rotation.W = Value;
+								break;
+							}
+						}
+						Rotation.Normalize();
+						Transform.SetRotation(Rotation);
+						break;
+					}
+					case ESlateRotationRepresentation::AxisAndAngle:
+					{
+						FVector Axis;
+						NumericType Angle;
+						Transform.GetRotation().ToAxisAndAngle(Axis, Angle);
+
+						switch(SubComponent)
+						{
+							case ESlateTransformSubComponent::X:
+							{
+								Axis.X = Value;
+								break;
+							}
+							case ESlateTransformSubComponent::Y:
+							{
+								Axis.Y = Value;
+								break;
+							}
+							case ESlateTransformSubComponent::Z:
+							{
+								Axis.Z = Value;
+								break;
+							}
+							case ESlateTransformSubComponent::Angle:
+							{
+								Angle = Value;
+								break;
+							}
+						}
+						Transform.SetRotation(FQuat(Axis, Angle));
+						break;
+					}
+				}
+				break;
+			}
+			case ESlateTransformComponent::Scale:
+			{
+				UE::Math::TVector<NumericType> Scale3D = Transform.GetScale3D();
+				switch(SubComponent)
+				{
+					case ESlateTransformSubComponent::X:
+					{
+						Scale3D.X = Value;
+						break;
+					}
+					case ESlateTransformSubComponent::Y:
+					{
+						Scale3D.Y = Value;
+						break;
+					}
+					case ESlateTransformSubComponent::Z:
+					{
+						Scale3D.Z = Value;
+						break;
+					}
+				}
+				Transform.SetScale3D(Scale3D);
+				break;
+			}
+		}
+	}
 };
diff --git a/Engine/Source/Editor/SkeletonEditor/Private/BoneProxy.cpp b/Engine/Source/Editor/SkeletonEditor/Private/BoneProxy.cpp
--- a/Engine/Source/Editor/SkeletonEditor/Private/BoneProxy.cpp
+++ b/Engine/Source/Editor/SkeletonEditor/Private/BoneProxy.cpp
@@ -117,50 +117,14 @@ TOptional<FVector::FReal> UBoneProxy::GetNumericValue(
 			break;
 		}
 	}
-	
-	switch(Component)
-	{
-		case ESlateTransformComponent::Location:
-		{
-			switch(SubComponent)
-			{
-				case ESlateTransformSubComponent::X:
-					return LocationPtr->X;
-				case ESlateTransformSubComponent::Y:
-					return LocationPtr->Y;
-				case ESlateTransformSubComponent::Z:
-					return LocationPtr->Z;
-			}
-			break;
-		}
-		case ESlateTransformComponent::Rotation:
-		{
-			switch(SubComponent)
-			{
-			case ESlateTransformSubComponent::Roll:
-				return RotationPtr->Roll;
-			case ESlateTransformSubComponent::Pitch:
-				return RotationPtr->Pitch;
-			case ESlateTransformSubComponent::Yaw:
-				return RotationPtr->Yaw;
-			}
-			break;
-		}
-		case ESlateTransformComponent::Scale:
-		{
-			switch(SubComponent)
-			{
-			case ESlateTransformSubComponent::X:
-				return ScalePtr->X;
-			case ESlateTransformSubComponent::Y:
-				return ScalePtr->Y;
-			case ESlateTransformSubComponent::Z:
-				return ScalePtr->Z;
-			}
-			break;
-		}
-	}
-	return TOptional<FVector::FReal>();
+
+	const FEulerTransform Transform(*RotationPtr, *LocationPtr, *ScalePtr);
+	return SAdvancedTransformInputBox<FEulerTransform>::GetNumericValueFromTransform(
+		Transform,
+		Component,
+		Representation,
+		SubComponent
+		);
 }
 
 TOptional<FVector::FReal> UBoneProxy::GetMultiNumericValue(
