Fix non-unity build issue.

diff --git a/Engine/Source/Developer/AnimationWidgets/Public/SAdvancedTransformInputBox.h b/Engine/Source/Developer/AnimationWidgets/Public/SAdvancedTransformInputBox.h
--- a/Engine/Source/Developer/AnimationWidgets/Public/SAdvancedTransformInputBox.h
+++ b/Engine/Source/Developer/AnimationWidgets/Public/SAdvancedTransformInputBox.h
@@ -4,6 +4,7 @@
 
 #include "SAdvancedRotationInputBox.h"
 #include "Widgets/Input/SComboButton.h"
+#include "Widgets/Input/SButton.h"
 #include "Styling/AppStyle.h"
 #include "Framework/Application/SlateApplication.h"
 
