Fix non-unity

diff --git a/Engine/Plugins/Animation/IKRig/Source/IKRig/Public/IKRigDefinition.h b/Engine/Plugins/Animation/IKRig/Source/IKRig/Public/IKRigDefinition.h
--- a/Engine/Plugins/Animation/IKRig/Source/IKRig/Public/IKRigDefinition.h
+++ b/Engine/Plugins/Animation/IKRig/Source/IKRig/Public/IKRigDefinition.h
@@ -11,6 +11,7 @@
 #endif
 #include "IKRigDefinition.generated.h"
 
+class IPropertyHandle;
 class UIKRigSolver;
 
 #if WITH_EDITORONLY_DATA
diff --git a/Engine/Plugins/Animation/IKRig/Source/IKRigEditor/Private/RigEditor/IKRigEditorController.cpp b/Engine/Plugins/Animation/IKRig/Source/IKRigEditor/Private/RigEditor/IKRigEditorController.cpp
--- a/Engine/Plugins/Animation/IKRig/Source/IKRigEditor/Private/RigEditor/IKRigEditorController.cpp
+++ b/Engine/Plugins/Animation/IKRig/Source/IKRigEditor/Private/RigEditor/IKRigEditorController.cpp
@@ -57,6 +57,7 @@ TOptional<FTransform> UIKRigBoneDetails::GetTransform(EIKRigTransformType::Type
 			break;
 		}
 	}
+	checkSlow(IsRelative);
 
 	FTransform Transform = LocalTransform;
 	if(!IsRelative[0]) Transform.SetLocation(GlobalTransform.GetLocation());
diff --git a/Engine/Source/Developer/AnimationWidgets/Public/SAdvancedTransformInputBox.h b/Engine/Source/Developer/AnimationWidgets/Public/SAdvancedTransformInputBox.h
--- a/Engine/Source/Developer/AnimationWidgets/Public/SAdvancedTransformInputBox.h
+++ b/Engine/Source/Developer/AnimationWidgets/Public/SAdvancedTransformInputBox.h
@@ -6,7 +6,10 @@
 #include "Widgets/Input/SComboButton.h"
 #include "Widgets/Input/SButton.h"
 #include "Styling/AppStyle.h"
+#include "Textures/SlateIcon.h"
 #include "Framework/Application/SlateApplication.h"
+#include "Framework/Commands/UIAction.h"
+#include "Framework/Commands/UICommandInfo.h"
 
 namespace ESlateTransformComponent
 {
