Fix for duplicate Module name FToolWidgetsModule. Renames to FAnimationWidgetsModule

diff --git a/Engine/Source/Developer/AnimationWidgets/Private/AnimationWidgetsModule.cpp b/Engine/Source/Developer/AnimationWidgets/Private/AnimationWidgetsModule.cpp
--- a/Engine/Source/Developer/AnimationWidgets/Private/AnimationWidgetsModule.cpp
+++ b/Engine/Source/Developer/AnimationWidgets/Private/AnimationWidgetsModule.cpp
@@ -3,7 +3,7 @@
 #include "Modules/ModuleInterface.h"
 #include "Modules/ModuleManager.h"
 
-class FToolWidgetsModule
+class FAnimationWidgetsModule
 	: public IModuleInterface
 {
 public:
@@ -15,4 +15,4 @@ public:
 };
 
 
-IMPLEMENT_MODULE(FToolWidgetsModule, ToolWidgets);
+IMPLEMENT_MODULE(FAnimationWidgetsModule, AnimationWidgets);
