Fixed crash in new advanced transform details panels.

diff --git a/Engine/Source/Developer/AnimationWidgets/Public/SAdvancedRotationInputBox.h b/Engine/Source/Developer/AnimationWidgets/Public/SAdvancedRotationInputBox.h
--- a/Engine/Source/Developer/AnimationWidgets/Public/SAdvancedRotationInputBox.h
+++ b/Engine/Source/Developer/AnimationWidgets/Public/SAdvancedRotationInputBox.h
@@ -58,7 +58,7 @@ class SAdvancedRotationInputBox : public SCompoundWidget
 public:
 	
 	typedef SNumericVectorInputBox<NumericType, UE::Math::TVector<NumericType>, 3> SNumericVectorInputBox3; 
-	typedef SNumericVectorInputBox<NumericType, UE::Math::TVector<NumericType>, 4> SNumericVectorInputBox4; 
+	typedef SNumericVectorInputBox<NumericType, UE::Math::TVector4<NumericType>, 4> SNumericVectorInputBox4; 
 
 	/** Delegate for notification for a rotator value change */
  	DECLARE_DELEGATE_OneParam(FOnRotatorChanged, FRotator);
@@ -243,7 +243,7 @@ public:
 				.TogglePadding(InArgs._TogglePadding)
 				.Visibility(this, &SAdvancedRotationInputBox::IsEulerInputBoxVisible)
 			]
-
+			
 			// a vector 4d variation - to represent quaternion
 			+SHorizontalBox::Slot()
 			.FillWidth(1.f)
