Control Rig: Add Schematic Graph Panel
#rnx

diff --git a/Engine/Plugins/Animation/ControlRig/Source/ControlRigEditor/Private/ControlRigEditorStyle.h b/Engine/Plugins/Animation/ControlRig/Source/ControlRigEditor/Private/ControlRigEditorStyle.h
--- a/Engine/Plugins/Animation/ControlRig/Source/ControlRigEditor/Private/ControlRigEditorStyle.h
+++ b/Engine/Plugins/Animation/ControlRig/Source/ControlRigEditor/Private/ControlRigEditorStyle.h
@@ -141,6 +141,11 @@ public:
 			Set("ControlRig.TestData.Record", new IMAGE_BRUSH("Slate/RecordingIndicator", Icon32x32));
 		}
 
+		// Schematic
+		{
+			Set("ControlRig.Schematic.Circle", new IMAGE_BRUSH_SVG("Slate/Socket_Closed", Icon16x16));
+		}
+
 		FSlateStyleRegistry::RegisterSlateStyle(*this);
 	}
 
diff --git a/Engine/Plugins/Animation/ControlRig/Source/ControlRigEditor/Private/Editor/ControlRigEditor.cpp b/Engine/Plugins/Animation/ControlRig/Source/ControlRigEditor/Private/Editor/ControlRigEditor.cpp
--- a/Engine/Plugins/Animation/ControlRig/Source/ControlRigEditor/Private/Editor/ControlRigEditor.cpp
+++ b/Engine/Plugins/Animation/ControlRig/Source/ControlRigEditor/Private/Editor/ControlRigEditor.cpp
@@ -106,12 +106,14 @@
 #include "RigVMFunctions/RigVMFunction_ControlFlow.h"
 #include "RigVMModel/Nodes/RigVMAggregateNode.h"
 #include "AnimationEditorViewportClient.h"
+#include "SchematicGraphPanel/SSchematicGraphPanel.h"
 #include "RigVMCore/RigVMExecuteContext.h"
 #include "Editor/RigVMGraphDetailCustomization.h"
 
 #define LOCTEXT_NAMESPACE "ControlRigEditor"
 
 TAutoConsoleVariable<bool> CVarControlRigShowTestingToolbar(TEXT("ControlRig.Test.EnableTestingToolbar"), false, TEXT("When true we'll show the testing toolbar in Control Rig Editor."));
+TAutoConsoleVariable<bool> CVarShowSchematicPanelOverlay(TEXT("ControlRig.Preview.ShowSchematicPanelOverlay"), false, TEXT("When true we'll add an overlay to the persona viewport to show modular rig information."));
 
 const FName FControlRigEditorModes::ControlRigEditorMode = TEXT("Rigging");
 const TArray<FName> FControlRigEditor::ForwardsSolveEventQueue = {FRigUnit_BeginExecution::EventName};
@@ -1783,6 +1785,8 @@ void FControlRigEditor::HandleViewportCreated(const TSharedRef<class IPersonaVie
 {
 	DECLARE_SCOPE_HIERARCHICAL_COUNTER_FUNC()
 
+	PreviewViewport = InViewport;
+
 	// TODO: this is duplicated code from FAnimBlueprintEditor, would be nice to consolidate. 
 	auto GetCompilationStateText = [this]()
 	{
@@ -2218,6 +2222,21 @@ void FControlRigEditor::HandleViewportCreated(const TSharedRef<class IPersonaVie
         .Padding(0.0f)
         .ShowEffectWhenDisabled(false)
 	);
+
+	if (CVarShowSchematicPanelOverlay->GetBool())
+	{
+		if (UControlRigBlueprint* Blueprint = GetControlRigBlueprint())
+		{
+			if (Blueprint->IsModularRig())
+			{
+				TSharedRef<SSchematicGraphPanel> SchematicPanel = SNew(SSchematicGraphPanel)
+																.GraphData(&SchematicGraph)
+																.IsOverlay(true);
+				InViewport->AddOverlayWidget(SchematicPanel);
+				HandleSchematicViewportCreated(SchematicPanel);
+			}
+		}
+	}
 	
 	InViewport->GetKeyDownDelegate().BindLambda([&](const FGeometry& MyGeometry, const FKeyEvent& InKeyEvent) -> FReply {
 		if (OnKeyDownDelegate.IsBound())
@@ -2506,6 +2525,105 @@ void FControlRigEditor::CacheNameLists()
 	}
 }
 
+void FControlRigEditor::HandleSchematicViewportCreated(const TSharedRef<SSchematicGraphPanel>& InViewport)
+{
+	InViewport->UpdateNodeWidgetDelegate.BindSP(this, &FControlRigEditor::HandleUpdateSchematicNodes);
+	InViewport->OnNodeClickedDelegate.BindSP(this, &FControlRigEditor::HandleSchematicNodeClicked);
+}
+
+FVector2D FControlRigEditor::ComputePersonaProjectedScreenPos(const FVector& InWorldPos)
+{
+	if (PreviewViewport.IsValid())
+	{
+		FEditorViewportClient& Client = PreviewViewport->GetViewportClient();
+		FSceneViewFamilyContext ViewFamily(FSceneViewFamily::ConstructionValues(
+							Client.Viewport,
+							Client.GetScene(),
+							Client.EngineShowFlags));
+		// SceneView is deleted with the ViewFamily
+		FSceneView* SceneView = Client.CalcSceneView(&ViewFamily);
+	
+		// Compute the MinP/MaxP in pixel coord, relative to View.ViewRect.Min
+		const FMatrix& WorldToView = SceneView->ViewMatrices.GetViewMatrix();
+		const FMatrix& ViewToProj = SceneView->ViewMatrices.GetProjectionMatrix();
+		const float NearClippingDistance = SceneView->NearClippingDistance + SMALL_NUMBER;
+		const FIntRect ViewRect = SceneView->UnconstrainedViewRect;
+
+		// Clamp position on the near plane to get valid rect even if bounds' points are behind the camera
+		FPlane P_View = WorldToView.TransformFVector4(FVector4(InWorldPos, 1.f));
+		if (P_View.Z <= NearClippingDistance)
+		{
+			P_View.Z = NearClippingDistance;
+		}
+
+		// Project from view to projective space
+		FVector2D MinP(FLT_MAX, FLT_MAX);
+		FVector2D MaxP(-FLT_MAX, -FLT_MAX);
+		FVector2D ScreenPos;
+		const bool bIsValid = FSceneView::ProjectWorldToScreen(P_View, ViewRect, ViewToProj, ScreenPos);
+
+		// Clamp to pixel border
+		ScreenPos = FIntPoint(FMath::FloorToInt(ScreenPos.X), FMath::FloorToInt(ScreenPos.Y));
+
+		// Clamp to screen rect
+		ScreenPos.X = FMath::Clamp(ScreenPos.X, ViewRect.Min.X, ViewRect.Max.X);
+		ScreenPos.Y = FMath::Clamp(ScreenPos.Y, ViewRect.Min.Y, ViewRect.Max.Y);
+
+		return FVector2D(ScreenPos.X, ScreenPos.Y);
+	}
+	return FVector2D::ZeroVector;
+}
+
+void FControlRigEditor::HandleUpdateSchematicNodes(SSchematicGraphPanel* InPanel, TSharedPtr<SSchematicGraphNode> InNode)
+{
+	UControlRigBlueprint* Blueprint = GetControlRigBlueprint();
+	if (!Blueprint)
+	{
+		return;
+	}
+
+	FRigModuleReference* Module = Blueprint->ModularRigModel.FindModule(InNode->NodeData->Name);
+	if (!Module)
+	{
+		return;
+	}
+
+	FString Path = Module->GetPath();
+	int32 Level=0;
+	FRigModuleReference* CurModule = Module;
+	while (CurModule && !CurModule->ParentPath.IsEmpty())
+	{
+		Level++;
+		CurModule = Blueprint->ModularRigModel.FindModule(CurModule->ParentPath);
+	}
+
+	if (InPanel->bIsOverlay)
+	{
+		FTransform Transform;
+		if (UControlRig* DebuggedControlRig = Cast<UControlRig>(GetBlueprintObj()->GetObjectBeingDebugged()))
+		{
+			Transform = DebuggedControlRig->GetHierarchy()->GetGlobalTransform(FRigElementKey("hand_l", ERigElementType::Bone));
+			FVector2D PixelPos = ComputePersonaProjectedScreenPos(Transform.GetLocation());
+			InNode->SetPosition(PixelPos, true);
+		}
+
+		InNode->Brush = *FControlRigEditorStyle::Get().GetBrush( "ControlRig.Schematic.Circle");
+	}
+	else
+	{
+		InNode->SetPosition(FVector2D(100, 100*Level));
+		InNode->Brush = *FAppStyle::GetBrush("WhiteTexture");
+	}
+}
+
+void FControlRigEditor::HandleSchematicNodeClicked(SSchematicGraphPanel* InPanel, SSchematicGraphNode* InNode)
+{
+	for (FSchematicGraphNode* Node : SchematicGraph.Nodes)
+	{
+		Node->bIsSelected = Node == InNode->NodeData;
+	}
+}
+
 void FControlRigEditor::HandlePreviewMeshChanged(USkeletalMesh* InOldSkeletalMesh, USkeletalMesh* InNewSkeletalMesh)
 {
 	RebindToSkeletalMeshComponent();
@@ -3279,6 +3397,10 @@ void FControlRigEditor::OnHierarchyModified(ERigHierarchyNotification InNotif, U
 						RigBlueprint->TurnIntoControlRigModule();
 					}
 				}
+				else if(InElement->GetType() == ERigElementType::Socket)
+				{
+					SchematicGraph.AddNode(InElement->GetName());
+				}
 			}
 			// no break - fall through
 		}
@@ -3638,6 +3760,8 @@ void FControlRigEditor::HandleModularRigModified(EModularRigNotification InNotif
 			{
 				ClearDetailObject();
 			}
+
+			// todo: update SchematicGraph
 			break;
 		}
 		case EModularRigNotification::ModuleReparented:
@@ -3654,6 +3778,13 @@ void FControlRigEditor::HandleModularRigModified(EModularRigNotification InNotif
 			}
 			ModulesSelected.Remove(OldPath);
 			ModulesSelected.Add(InModule->GetPath());
+
+			// todo: update SchematicGraph
+			break;
+		}
+		case EModularRigNotification::ConnectionChanged:
+		{
+			// todo: update SchematicGraph
 			break;
 		}
 	}
diff --git a/Engine/Plugins/Animation/ControlRig/Source/ControlRigEditor/Private/Editor/ControlRigEditor.h b/Engine/Plugins/Animation/ControlRig/Source/ControlRigEditor/Private/Editor/ControlRigEditor.h
--- a/Engine/Plugins/Animation/ControlRig/Source/ControlRigEditor/Private/Editor/ControlRigEditor.h
+++ b/Engine/Plugins/Animation/ControlRig/Source/ControlRigEditor/Private/Editor/ControlRigEditor.h
@@ -24,6 +24,7 @@
 #include "ControlRigTestData.h"
 #include "ModularRigController.h"
 #include "RigVMHost.h"
+#include "SchematicGraphPanel/SSchematicGraphPanel.h"
 #include "Units/RigUnit.h"
 
 class UControlRigBlueprint;
@@ -240,6 +241,12 @@ private:
 	TOptional<float> GetToolbarAxesScale() const;
 	void OnToolbarAxesScaleChanged(float InValue);
 
+	/** Handle schematic panel setup*/
+	void HandleSchematicViewportCreated(const TSharedRef<class SSchematicGraphPanel>& InViewport);
+	void HandleUpdateSchematicNodes(SSchematicGraphPanel* InPanel, TSharedPtr<SSchematicGraphNode> InNode);
+	void HandleSchematicNodeClicked(SSchematicGraphPanel* InPanel, SSchematicGraphNode* InNode);
+	FVector2D ComputePersonaProjectedScreenPos(const FVector& InWorldPos);
+
 		/** Handle switching skeletal meshes */
 	void HandlePreviewMeshChanged(USkeletalMesh* InOldSkeletalMesh, USkeletalMesh* InNewSkeletalMesh);
 
@@ -287,7 +294,7 @@ protected:
 	TSharedPtr<IPersonaToolkit> PersonaToolkit;
 
 	/** Preview instance inspector widget */
-	TSharedPtr<SWidget> PreviewEditor;
+	TSharedPtr<IPersonaViewport> PreviewViewport;
 
 	/** preview scene */
 	TSharedPtr<IPersonaPreviewScene> PreviewScene;
@@ -295,6 +302,9 @@ protected:
 	/** preview animation instance */
 	UAnimPreviewInstance* PreviewInstance;
 
+	/** Model for the schematic views */
+	FSchematicGraph SchematicGraph;
+
 	/** Delegate to deal with key down evens in the viewport / editor */
 	FPersonaViewportKeyDownDelegate OnKeyDownDelegate;
 
@@ -350,6 +360,7 @@ protected:
 	static const TArray<FName> BackwardsAndForwardsSolveEventQueue;
 
 	friend class FControlRigEditorMode;
+	friend class FModularRigEditorMode;
 	friend class SControlRigStackView;
 	friend class SRigHierarchy;
 	friend class SModularRigModel;
diff --git a/Engine/Source/Developer/AnimationWidgets/AnimationWidgets.Build.cs b/Engine/Source/Developer/AnimationWidgets/AnimationWidgets.Build.cs
--- a/Engine/Source/Developer/AnimationWidgets/AnimationWidgets.Build.cs
+++ b/Engine/Source/Developer/AnimationWidgets/AnimationWidgets.Build.cs
@@ -25,6 +25,10 @@ public class AnimationWidgets : ModuleRules
 			PrivateDependencyModuleNames.AddRange(
 				new string[]
 				{
+					"Engine",
+					"CoreUObject",
+					"UnrealEd",
+					"GraphEditor",
 					"PropertyEditor"
 				}
 			);
diff --git a/Engine/Source/Developer/AnimationWidgets/Private/SchematicGraphPanel/SSchematicGraphPanel.cpp b/Engine/Source/Developer/AnimationWidgets/Private/SchematicGraphPanel/SSchematicGraphPanel.cpp
new file mode 100644
--- /dev/null
+++ b/Engine/Source/Developer/AnimationWidgets/Private/SchematicGraphPanel/SSchematicGraphPanel.cpp
@@ -0,0 +1,246 @@
+// Copyright Epic Games, Inc. All Rights Reserved.
+#if WITH_EDITOR
+
+#include "SchematicGraphPanel/SSchematicGraphPanel.h"
+#include "SGraphPanel.h"
+
+#define LOCTEXT_NAMESPACE "SSchematicGraphPanel"
+
+
+bool FSchematicGraph::AddNode(const FString& InName)
+{
+	if (Nodes.ContainsByPredicate([InName](const FSchematicGraphNode* Node)
+	{
+		return Node->Name == InName;
+	}))
+	{
+		return false;
+	}
+
+	FSchematicGraphNode* NewElement = static_cast<FSchematicGraphNode*>(FMemory::Malloc(sizeof(FSchematicGraphNode)));
+	new (NewElement) FSchematicGraphNode();
+	NewElement->Name = InName;
+	Nodes.Add(NewElement);
+	OnNodeAddedDelegate.Broadcast(NewElement);
+	return true;
+}
+
+void SSchematicGraphNode::Construct(const FArguments& InArgs)
+{
+	NodeData = InArgs._NodeData;
+	OnClickedDelegate = InArgs._OnClicked;
+	TEasingAttributeInterpolator<FVector2d>::FSettings Vector2DInterpSettings(EEasingInterpolatorType::BounceEaseInOut, 0.2f);
+	Position = TAnimatedAttribute<FVector2d>::Create(Vector2DInterpSettings, FVector2d::ZeroVector);
+	Size = TAnimatedAttribute<FVector2d>::Create(Vector2DInterpSettings, OriginalSize);
+	TEasingAttributeInterpolator<float>::FSettings FloatInterpSettings(EEasingInterpolatorType::BounceEaseInOut, 0.2f);
+	Scale = TAnimatedAttribute<float>::Create(FloatInterpSettings, 1.f);
+	Brush = *FAppStyle::GetBrush("WhiteTexture");
+	Brush.TintColor = EStyleColor::AccentWhite;
+}
+
+FVector2D SSchematicGraphNode::ComputeDesiredSize(float LayoutScaleMultiplier) const
+{
+	return OriginalSize*LayoutScaleMultiplier;
+}
+
+int32 SSchematicGraphNode::OnPaint(const FPaintArgs& Args, const FGeometry& AllottedGeometry, const FSlateRect& MyCullingRect, FSlateWindowElementList& OutDrawElements, int32 LayerId, const FWidgetStyle& InWidgetStyle, bool bParentEnabled) const
+{
+	int32 NewLayerId = SNodePanel::SNode::OnPaint(Args, AllottedGeometry, MyCullingRect, OutDrawElements, LayerId, InWidgetStyle, bParentEnabled);
+	NewLayerId++;
+
+	//FVector2d CenterOffset = OriginalSize*-0.5;
+	FVector2d CurSize = Size->Get() * Scale->Get();
+	FVector2d SizeOffset = (CurSize-OriginalSize)*-0.5;
+	//FVector2d TotalOffset = CenterOffset + SizeOffset;
+	FVector2d TotalOffset = SizeOffset;
+
+	FLinearColor Color = NodeData->bIsSelected ?
+		FSlateColor(EStyleColor::AccentRed).GetSpecifiedColor()
+	: FLinearColor::White;
+
+	FSlateDrawElement::MakeBox(
+				OutDrawElements,
+				NewLayerId,
+				AllottedGeometry.ToPaintGeometry(CurSize, FSlateLayoutTransform(TotalOffset)),
+				&Brush,
+				ESlateDrawEffect::None,
+				Color
+				);
+	
+	return NewLayerId;
+}
+
+void SSchematicGraphNode::OnMouseEnter(const FGeometry& MyGeometry, const FPointerEvent& MouseEvent)
+{
+	SNode::OnMouseEnter(MyGeometry, MouseEvent);
+	Scale->Set(ScaledUp);
+}
+
+void SSchematicGraphNode::OnMouseLeave(const FPointerEvent& MouseEvent)
+{
+	SNode::OnMouseLeave(MouseEvent);
+	Scale->Set(1.f);
+}
+
+FReply SSchematicGraphNode::OnDragOver(const FGeometry& MyGeometry, const FDragDropEvent& DragDropEvent)
+{
+	SNode::OnDragOver(MyGeometry, DragDropEvent);
+	Scale->Set(ScaledDown);
+	return FReply::Handled();
+}
+
+void SSchematicGraphNode::OnDragLeave(const FDragDropEvent& DragDropEvent)
+{
+	SNode::OnDragLeave(DragDropEvent);
+	Scale->Set(1.f);
+}
+
+FReply SSchematicGraphNode::OnMouseButtonDown(const FGeometry& MyGeometry, const FPointerEvent& MouseEvent)
+{
+	SNode::OnMouseButtonDown(MyGeometry, MouseEvent);
+	if (MouseEvent.GetPressedButtons().Contains(EKeys::LeftMouseButton))
+	{
+		OnClickedDelegate.ExecuteIfBound(this);
+	}
+	return FReply::Handled();
+}
+
+FVector2d SSchematicGraphNode::GetPosition() const
+{
+	return Position->Get() - (OriginalSize*0.5);
+}
+
+void SSchematicGraphPanel::SetSchematicGraph(FSchematicGraph* InGraphData)
+{
+	if (GraphData)
+	{
+		GraphData->OnNodeAdded().RemoveAll(this);
+	}
+	
+	GraphData = InGraphData;
+	if (GraphData)
+	{
+		GraphData->OnNodeAdded().AddSP(this, &SSchematicGraphPanel::AddNode);
+	}
+}
+
+void SSchematicGraphPanel::Construct(const FArguments& InArgs)
+{
+	GraphData = InArgs._GraphData;
+	bIsOverlay = InArgs._IsOverlay;
+	UpdateNodeWidgetDelegate = InArgs._OnUpdateNodeWidget;
+	OnNodeClickedDelegate = InArgs._OnNodeClicked;
+	
+	SNodePanel::Construct();
+
+	SetVisibility(bIsOverlay ? EVisibility::SelfHitTestInvisible : EVisibility::Visible);
+	if (GraphData)
+	{
+		GraphData->OnNodeAdded().AddSP(this, &SSchematicGraphPanel::AddNode);
+	}
+}
+
+void SSchematicGraphPanel::RebuildPanel()
+{
+	RemoveAllNodes();
+
+	if (GraphData)
+	{
+		for (FSchematicGraphNode* Node : GraphData->Nodes)
+		{
+			AddNode(Node);
+		}
+	}
+}
+
+void SSchematicGraphPanel::AddNode(FSchematicGraphNode* NodeToAdd)
+{
+	TSharedRef<SSchematicGraphNode> NewNode = SNew(SSchematicGraphNode)
+														.OnClicked_Raw(this, &SSchematicGraphPanel::OnNodeClicked)
+														.NodeData(NodeToAdd);
+	SNodePanel::AddGraphNode(NewNode);
+}
+
+int32 SSchematicGraphPanel::OnPaint(const FPaintArgs& Args, const FGeometry& AllottedGeometry, const FSlateRect& MyCullingRect, FSlateWindowElementList& OutDrawElements, int32 LayerId, const FWidgetStyle& InWidgetStyle, bool bParentEnabled) const
+{
+	if (!bIsOverlay)
+	{
+		const FSlateBrush* DefaultBackground = FAppStyle::GetBrush(TEXT("Graph.Panel.SolidBackground"));
+		PaintBackgroundAsLines(DefaultBackground, AllottedGeometry, MyCullingRect, OutDrawElements, LayerId);
+	}
+	const int32 NodeLayerId = LayerId + 1;
+	int32 MaxLayerId = NodeLayerId;
+
+	if (!GraphData)
+	{
+		return MaxLayerId; 
+	}
+
+	FArrangedChildren ArrangedChildren(EVisibility::Visible);
+	ArrangeChildNodes(AllottedGeometry, ArrangedChildren);
+
+	// Because we paint multiple children, we must track the maximum layer id that they produced in case one of our parents
+	// wants to an overlay for all of its contents.
+
+	const FPaintArgs NewArgs = Args.WithNewParent(this);
+
+	// Draw the child nodes
+	{
+		for (int32 ChildIndex = 0; ChildIndex < ArrangedChildren.Num(); ++ChildIndex)
+		{
+			FArrangedWidget& CurWidget = ArrangedChildren[ChildIndex];
+			TSharedRef<SSchematicGraphNode> ChildNode = StaticCastSharedRef<SSchematicGraphNode>(CurWidget.Widget);
+			
+			// Examine node to see what layers we should be drawing in
+			int32 ChildLayerId = NodeLayerId;
+
+			const bool bNodeIsVisible = FSlateRect::DoRectanglesIntersect( CurWidget.Geometry.GetLayoutBoundingRect(), MyCullingRect );
+
+			if (bNodeIsVisible)
+			{
+				int32 CurWidgetsMaxLayerId = CurWidget.Widget->Paint(NewArgs, CurWidget.Geometry, MyCullingRect, OutDrawElements, ChildLayerId, InWidgetStyle, true );
+				MaxLayerId = FMath::Max( MaxLayerId, CurWidgetsMaxLayerId + 1 );
+			}
+		}
+	}
+
+	// Draw the software cursor
+	++MaxLayerId;
+	PaintSoftwareCursor(AllottedGeometry, MyCullingRect, OutDrawElements, MaxLayerId);
+
+	return MaxLayerId;
+}
+
+TSharedRef<SSchematicGraphNode> SSchematicGraphPanel::GetChild(int32 ChildIndex) const
+{
+	return StaticCastSharedRef<SSchematicGraphNode>(Children[ChildIndex]);
+}
+
+TStatId SSchematicGraphPanel::GetStatId() const
+{
+	RETURN_QUICK_DECLARE_CYCLE_STAT(SSchematicGraphPanel, STATGROUP_Tickables);
+}
+
+void SSchematicGraphPanel::Tick(float DeltaTime)
+{
+	for (int32 i=0; i<Children.Num(); ++i)
+	{
+		TSharedRef<SSchematicGraphNode> SNode = GetChild(i);
+		UpdateNodeWidgetDelegate.Execute(this, SNode.ToSharedPtr());
+	}
+}
+
+void SSchematicGraphPanel::Tick(const FGeometry& AllottedGeometry, const double InCurrentTime, const float InDeltaTime)
+{
+	Tick(InDeltaTime);
+}
+
+void SSchematicGraphPanel::OnNodeClicked(SSchematicGraphNode* Node)
+{
+	OnNodeClickedDelegate.ExecuteIfBound(this, Node);
+}
+
+
+#undef LOCTEXT_NAMESPACE
+
+#endif
\ No newline at end of file
diff --git a/Engine/Source/Developer/AnimationWidgets/Private/SchematicGraphPanel/SchematicViewportTabSummoner.cpp b/Engine/Source/Developer/AnimationWidgets/Private/SchematicGraphPanel/SchematicViewportTabSummoner.cpp
new file mode 100644
--- /dev/null
+++ b/Engine/Source/Developer/AnimationWidgets/Private/SchematicGraphPanel/SchematicViewportTabSummoner.cpp
@@ -0,0 +1,47 @@
+// Copyright Epic Games, Inc. All Rights Reserved.
+#if WITH_EDITOR
+
+#include "SchematicGraphPanel/SchematicViewportTabSummoner.h"
+#include "SchematicGraphPanel/SSchematicGraphPanel.h"
+#include "Widgets/Docking/SDockTab.h"
+
+#define LOCTEXT_NAMESPACE "SchematicViewportTabSummoner"
+
+const FName FSchematicViewportTabSummoner::TabID(TEXT("SchematicViewport"));
+
+FSchematicViewportTabSummoner::FSchematicViewportTabSummoner(TSharedPtr<class FAssetEditorToolkit> InHostingApp, const FSchematicViewportArgs& InArgs)
+	: FWorkflowTabFactory(TabID, InHostingApp)
+	, SchematicGraph(InArgs.SchematicGraph)
+	, OnViewportCreated(InArgs.OnViewportCreated)
+{
+	TabLabel = LOCTEXT("SchematicViewportTabLabel", "Schematic Viewport");
+	TabIcon = FSlateIcon(FAppStyle::GetAppStyleSetName(), "LevelEditor.Tabs.Viewports");
+
+	ViewMenuDescription = LOCTEXT("SchematicViewport_ViewMenu_Desc", "Schematic Viewport");
+	ViewMenuTooltip = LOCTEXT("SchematicViewport_ViewMenu_ToolTip", "Show the Schematic Viewport tab");
+}
+
+FTabSpawnerEntry& FSchematicViewportTabSummoner::RegisterTabSpawner(TSharedRef<FTabManager> InTabManager, const FApplicationMode* CurrentApplicationMode) const
+{
+	FTabSpawnerEntry& SpawnerEntry = FWorkflowTabFactory::RegisterTabSpawner(InTabManager, CurrentApplicationMode);
+
+	SpawnerEntry.SetReuseTabMethod(FOnFindTabToReuse::CreateLambda([](const FTabId& InTabId) ->TSharedPtr<SDockTab> {
+	
+		return TSharedPtr<SDockTab>();
+
+	}));
+
+	return SpawnerEntry;
+}
+
+TSharedRef<SWidget> FSchematicViewportTabSummoner::CreateTabBody(const FWorkflowTabSpawnInfo& Info) const
+{
+	TSharedRef<SSchematicGraphPanel> Viewport = SNew(SSchematicGraphPanel)
+											.GraphData(SchematicGraph)
+											.IsOverlay(false);
+	OnViewportCreated.ExecuteIfBound(Viewport);
+	return Viewport;
+}
+
+#undef LOCTEXT_NAMESPACE 
+#endif
\ No newline at end of file
diff --git a/Engine/Source/Developer/AnimationWidgets/Public/SchematicGraphPanel/SSchematicGraphPanel.h b/Engine/Source/Developer/AnimationWidgets/Public/SchematicGraphPanel/SSchematicGraphPanel.h
new file mode 100644
--- /dev/null
+++ b/Engine/Source/Developer/AnimationWidgets/Public/SchematicGraphPanel/SSchematicGraphPanel.h
@@ -0,0 +1,117 @@
+// Copyright Epic Games, Inc. All Rights Reserved.
+#if WITH_EDITOR
+#pragma once
+
+#include "CoreMinimal.h"
+#include "SNodePanel.h"
+#include "TickableEditorObject.h"
+#include "Animation/AnimatedAttribute.h"
+
+struct ANIMATIONWIDGETS_API FSchematicGraphNode
+{
+	FString Name;
+	bool bIsSelected = false;
+};
+
+DECLARE_EVENT_OneParam(FSchematicGraph, FOnNodeAdded, FSchematicGraphNode*);
+
+class ANIMATIONWIDGETS_API FSchematicGraph
+{
+public:
+	TArray<FSchematicGraphNode*> Nodes;
+	TMap<FSchematicGraphNode*, TArray<FSchematicGraphNode*>> Links;
+
+	bool AddNode(const FString& InName);
+
+	FOnNodeAdded OnNodeAddedDelegate;
+	FOnNodeAdded& OnNodeAdded() { return OnNodeAddedDelegate; }
+};
+
+class ANIMATIONWIDGETS_API SSchematicGraphNode : public SNodePanel::SNode
+{
+public:
+	DECLARE_DELEGATE_OneParam(FOnClicked, SSchematicGraphNode*);
+
+	SLATE_BEGIN_ARGS(SSchematicGraphNode){}
+	SLATE_ARGUMENT(FSchematicGraphNode*, NodeData)
+	SLATE_EVENT(FOnClicked, OnClicked)
+	SLATE_END_ARGS()
+
+	void Construct(const FArguments& InArgs);
+
+	virtual FVector2D ComputeDesiredSize(float LayoutScaleMultiplier) const override;
+	virtual int32 OnPaint(const FPaintArgs& Args, const FGeometry& AllottedGeometry, const FSlateRect& MyCullingRect, FSlateWindowElementList& OutDrawElements, int32 LayerId, const FWidgetStyle& InWidgetStyle, bool bParentEnabled) const override;
+	
+	virtual void OnMouseEnter(const FGeometry& MyGeometry, const FPointerEvent& MouseEvent) override;
+	virtual void OnMouseLeave(const FPointerEvent& MouseEvent) override;
+	virtual FReply OnDragOver(const FGeometry& MyGeometry, const FDragDropEvent& DragDropEvent) override;
+	virtual void OnDragLeave(const FDragDropEvent& DragDropEvent) override;
+	virtual FReply OnMouseButtonDown(const FGeometry& MyGeometry, const FPointerEvent& MouseEvent) override;
+
+	virtual FVector2d GetPosition() const override;
+
+	void SetPosition(const FVector2d& InPosition, bool bImmediately = false)
+	{
+		bImmediately ?
+			Position->SetValueAndStop(FVector2d(InPosition))
+		: Position->Set(InPosition);
+	}
+	
+	FVector2d OriginalSize = FVector2d(50.0,50.0);
+	float ScaledUp = 1.25;
+	float ScaledDown = 0.75;
+
+	FOnClicked OnClickedDelegate;
+	FSchematicGraphNode* NodeData;
+	TSharedPtr<TAnimatedAttribute<FVector2d>> Position;
+	TSharedPtr<TAnimatedAttribute<FVector2d>> Size;
+	TSharedPtr<TAnimatedAttribute<float>> Scale;
+	FSlateBrush Brush;
+};
+
+
+/** Widget allowing editing of a control rig's structure */
+class ANIMATIONWIDGETS_API SSchematicGraphPanel : public SNodePanel, public FTickableEditorObject
+{
+public:
+	
+	DECLARE_DELEGATE_TwoParams(FUpdateNodeWidget, SSchematicGraphPanel*, TSharedPtr<SSchematicGraphNode>);
+	DECLARE_DELEGATE_TwoParams(FOnNodeClicked, SSchematicGraphPanel*, SSchematicGraphNode*);
+	
+	SLATE_BEGIN_ARGS(SSchematicGraphPanel) {}
+	SLATE_ARGUMENT(bool, IsOverlay)
+	SLATE_ARGUMENT(FSchematicGraph*, GraphData)
+	SLATE_EVENT(FUpdateNodeWidget, OnUpdateNodeWidget)
+	SLATE_EVENT(FOnNodeClicked, OnNodeClicked)
+	SLATE_END_ARGS()
+
+	~SSchematicGraphPanel(){}
+
+	void SetSchematicGraph(FSchematicGraph* InGraphData);
+	void Construct(const FArguments& InArgs);
+
+	void RebuildPanel();
+	void AddNode(FSchematicGraphNode* NodeToAdd);
+
+	// SNodePanel interface
+	virtual int32 OnPaint(const FPaintArgs& Args, const FGeometry& AllottedGeometry, const FSlateRect& MyCullingRect, FSlateWindowElementList& OutDrawElements, int32 LayerId, const FWidgetStyle& InWidgetStyle, bool bParentEnabled) const override;
+	// End of SNodePanel interface
+
+	TSharedRef<SSchematicGraphNode> GetChild(int32 ChildIndex) const;
+
+	// FTickableEditorObject Interface
+	virtual TStatId GetStatId() const override;
+	virtual void Tick(float DeltaTime) override;
+	virtual void Tick( const FGeometry& AllottedGeometry, const double InCurrentTime, const float InDeltaTime ) override;
+	virtual bool IsTickable() const override { return true; }
+	// End of FTickableEditorObject interface
+
+	void OnNodeClicked(SSchematicGraphNode* Node);
+
+	bool bIsOverlay;
+	FSchematicGraph* GraphData;
+	FUpdateNodeWidget UpdateNodeWidgetDelegate;
+	FOnNodeClicked OnNodeClickedDelegate;
+};
+
+#endif
\ No newline at end of file
diff --git a/Engine/Source/Developer/AnimationWidgets/Public/SchematicGraphPanel/SchematicViewportTabSummoner.h b/Engine/Source/Developer/AnimationWidgets/Public/SchematicGraphPanel/SchematicViewportTabSummoner.h
new file mode 100644
--- /dev/null
+++ b/Engine/Source/Developer/AnimationWidgets/Public/SchematicGraphPanel/SchematicViewportTabSummoner.h
@@ -0,0 +1,41 @@
+// Copyright Epic Games, Inc. All Rights Reserved.
+#if WITH_EDITOR
+
+#pragma once
+
+#include "CoreMinimal.h"
+#include "SSchematicGraphPanel.h"
+#include "WorkflowOrientedApp/WorkflowTabFactory.h"
+
+/** Called back when a viewport is created */
+DECLARE_DELEGATE_OneParam(FOnSchematicViewportCreated, const TSharedRef<class SSchematicGraphPanel>&);
+
+/** Arguments used to create a persona viewport tab */
+struct ANIMATIONWIDGETS_API FSchematicViewportArgs
+{
+	FSchematicViewportArgs()
+	{}
+
+	/** The model which contains the graph to display */
+	FSchematicGraph* SchematicGraph;
+	
+	/** Delegate fired when the viewport is created */
+	FOnSchematicViewportCreated OnViewportCreated;
+};
+
+
+struct ANIMATIONWIDGETS_API FSchematicViewportTabSummoner : public FWorkflowTabFactory
+{
+public:
+	static const FName TabID;
+	
+public:
+	FSchematicViewportTabSummoner(TSharedPtr<class FAssetEditorToolkit> InHostingApp, const FSchematicViewportArgs& InArgs);
+	
+	virtual FTabSpawnerEntry& RegisterTabSpawner(TSharedRef<FTabManager> TabManager, const FApplicationMode* CurrentApplicationMode) const;
+	virtual TSharedRef<SWidget> CreateTabBody(const FWorkflowTabSpawnInfo& Info) const override;
+
+	FSchematicGraph* SchematicGraph;
+	FOnSchematicViewportCreated OnViewportCreated;
+};
+#endif
\ No newline at end of file
diff --git a/Engine/Source/Runtime/SlateCore/Public/Animation/AnimatedAttribute.h b/Engine/Source/Runtime/SlateCore/Public/Animation/AnimatedAttribute.h
--- a/Engine/Source/Runtime/SlateCore/Public/Animation/AnimatedAttribute.h
+++ b/Engine/Source/Runtime/SlateCore/Public/Animation/AnimatedAttribute.h
@@ -42,7 +42,7 @@ public:
 	 */
 	template< typename InterpolatorSettings >
 	TAnimatedAttribute( FPrivateToken, const InterpolatorSettings& InSettings )
-		: Interpolator( MakeUnique<InterpolatorSettings::InterpolatorType>(InSettings) )
+		: Interpolator( MakeUnique<typename InterpolatorSettings::InterpolatorType>(InSettings) )
 		, Getter()
 		, DesiredValue() 
 	{
