Control Rig: Add asset drop ability to schematic panel
#rxn

diff --git a/Engine/Plugins/Animation/ControlRig/Source/ControlRigEditor/Private/ControlRigEditorStyle.h b/Engine/Plugins/Animation/ControlRig/Source/ControlRigEditor/Private/ControlRigEditorStyle.h
--- a/Engine/Plugins/Animation/ControlRig/Source/ControlRigEditor/Private/ControlRigEditorStyle.h
+++ b/Engine/Plugins/Animation/ControlRig/Source/ControlRigEditor/Private/ControlRigEditorStyle.h
@@ -143,7 +143,8 @@ public:
 
 		// Schematic
 		{
-			Set("ControlRig.Schematic.Circle", new IMAGE_BRUSH_SVG("Slate/Socket_Closed", Icon16x16));
+			Set("ControlRig.Schematic.SocketUnresolved", new IMAGE_BRUSH_SVG("Slate/Socket_Open", CoreStyleConstants::Icon128x128));
+			Set("ControlRig.Schematic.SocketResolved", new IMAGE_BRUSH_SVG("Slate/Socket_Closed", CoreStyleConstants::Icon128x128));
 		}
 
 		FSlateStyleRegistry::RegisterSlateStyle(*this);
diff --git a/Engine/Plugins/Animation/ControlRig/Source/ControlRigEditor/Private/Editor/ControlRigEditor.cpp b/Engine/Plugins/Animation/ControlRig/Source/ControlRigEditor/Private/Editor/ControlRigEditor.cpp
--- a/Engine/Plugins/Animation/ControlRig/Source/ControlRigEditor/Private/Editor/ControlRigEditor.cpp
+++ b/Engine/Plugins/Animation/ControlRig/Source/ControlRigEditor/Private/Editor/ControlRigEditor.cpp
@@ -106,6 +106,7 @@
 #include "RigVMFunctions/RigVMFunction_ControlFlow.h"
 #include "RigVMModel/Nodes/RigVMAggregateNode.h"
 #include "AnimationEditorViewportClient.h"
+#include "DragAndDrop/AssetDragDropOp.h"
 #include "SchematicGraphPanel/SSchematicGraphPanel.h"
 #include "RigVMCore/RigVMExecuteContext.h"
 #include "Editor/RigVMGraphDetailCustomization.h"
@@ -319,6 +320,12 @@ void FControlRigEditor::InitRigVMEditor(const EToolkitMode::Type Mode, const TSh
 			ControlRigBlueprint->OnHierarchyModified().AddSP(EditMode, &FControlRigEditMode::OnHierarchyModified_AnyThread);
 		}
 
+		if (ControlRigBlueprint->IsModularRig())
+		{
+			ControlRigBlueprint->OnSetObjectBeingDebugged().AddRaw(&SchematicModel, &FControlRigSchematicModel::OnSetObjectBeingDebugged);
+			ControlRigBlueprint->OnHierarchyModified().AddRaw(&SchematicModel, &FControlRigSchematicModel::OnHierarchyModified);
+		}
+
 		ControlRigBlueprint->OnRigTypeChanged().AddSP(this, &FControlRigEditor::HandleRigTypeChanged);
 		if (ControlRigBlueprint->IsModularRig())
 		{
@@ -2255,7 +2262,7 @@ void FControlRigEditor::HandleViewportCreated(const TSharedRef<class IPersonaVie
 			if (Blueprint->IsModularRig())
 			{
 				TSharedRef<SSchematicGraphPanel> SchematicPanel = SNew(SSchematicGraphPanel)
-																.GraphData(&SchematicGraph)
+																.GraphData(&SchematicModel.SchematicGraph)
 																.IsOverlay(true);
 				InViewport->AddOverlayWidget(SchematicPanel);
 				HandleSchematicViewportCreated(SchematicPanel);
@@ -2554,6 +2561,7 @@ void FControlRigEditor::HandleSchematicViewportCreated(const TSharedRef<SSchemat
 {
 	InViewport->UpdateNodeWidgetDelegate.BindSP(this, &FControlRigEditor::HandleUpdateSchematicNodes);
 	InViewport->OnNodeClickedDelegate.BindSP(this, &FControlRigEditor::HandleSchematicNodeClicked);
+	InViewport->OnDropDelegate.BindSP(this, &FControlRigEditor::HandleSchematicDrop);
 }
 
 FVector2D FControlRigEditor::ComputePersonaProjectedScreenPos(const FVector& InWorldPos)
@@ -2607,48 +2615,128 @@ void FControlRigEditor::HandleUpdateSchematicNodes(SSchematicGraphPanel* InPanel
 		return;
 	}
 
-	FRigModuleReference* Module = Blueprint->ModularRigModel.FindModule(InNode->NodeData->Name);
-	if (!Module)
-	{
-		return;
-	}
-
-	FString Path = Module->GetPath();
-	int32 Level=0;
-	FRigModuleReference* CurModule = Module;
-	while (CurModule && !CurModule->ParentPath.IsEmpty())
-	{
-		Level++;
-		CurModule = Blueprint->ModularRigModel.FindModule(CurModule->ParentPath);
-	}
-
 	if (InPanel->bIsOverlay)
 	{
 		FTransform Transform;
 		if (UControlRig* DebuggedControlRig = Cast<UControlRig>(GetBlueprintObj()->GetObjectBeingDebugged()))
 		{
-			Transform = DebuggedControlRig->GetHierarchy()->GetGlobalTransform(FRigElementKey("hand_l", ERigElementType::Bone));
+			FRigElementKey SocketKey(*InNode->NodeData->Name, ERigElementType::Socket);
+			Transform = DebuggedControlRig->GetHierarchy()->GetGlobalTransform(SocketKey);
 			FVector2D PixelPos = ComputePersonaProjectedScreenPos(Transform.GetLocation());
 			InNode->SetPosition(PixelPos, true);
+
+			
+			if (const FRigElementKey* ConnectorKey = DebuggedControlRig->ElementKeyRedirector.FindReverse(SocketKey))
+			{
+				InNode->Brush = *FControlRigEditorStyle::Get().GetBrush( "ControlRig.Schematic.SocketResolved");
+				InNode->Size->Set(InNode->OriginalSize*0.5);
+			}
+			else
+			{
+				InNode->Brush = *FControlRigEditorStyle::Get().GetBrush( "ControlRig.Schematic.SocketUnresolved");
+				InNode->Size->Set(InNode->OriginalSize);
+			}
+
+			if (FRigSocketElement* SocketElement = Cast<FRigSocketElement>(DebuggedControlRig->GetHierarchy()->Find(SocketKey)))
+			{
+				InNode->Brush.TintColor = SocketElement->GetColor(DebuggedControlRig->GetHierarchy());
+			}
+			else
+			{
+				InNode->Brush.TintColor = FStyleColors::AccentBlue;
+			}
 		}
 
-		InNode->Brush = *FControlRigEditorStyle::Get().GetBrush( "ControlRig.Schematic.Circle");
 	}
 	else
 	{
-		InNode->SetPosition(FVector2D(100, 100*Level));
+		InNode->SetPosition(FVector2D(0, 0));
 		InNode->Brush = *FAppStyle::GetBrush("WhiteTexture");
 	}
 }
 
 void FControlRigEditor::HandleSchematicNodeClicked(SSchematicGraphPanel* InPanel, SSchematicGraphNode* InNode)
 {
-	for (FSchematicGraphNode* Node : SchematicGraph.Nodes)
+	for (FSchematicGraphNode* Node : SchematicModel.SchematicGraph.Nodes)
 	{
 		Node->bIsSelected = Node == InNode->NodeData;
 	}
 }
 
+void FControlRigEditor::HandleSchematicDrop(SSchematicGraphPanel* InPanel, SSchematicGraphNode* InNode, const FDragDropEvent& InDragDropEvent)
+{
+	UControlRigBlueprint* Blueprint = GetControlRigBlueprint();
+	if (!Blueprint)
+	{
+		return;
+	}
+
+	UControlRig* ControlRig = Blueprint->GetDebuggedControlRig();
+	if (!ControlRig)
+	{
+		return;
+	}
+
+	URigHierarchy* Hierarchy = ControlRig->GetHierarchy();
+	if (!Hierarchy)
+	{
+		return;
+	}
+	
+	FRigBaseElement* Element = Hierarchy->Find(FRigElementKey(*InNode->NodeData->Name, ERigElementType::Socket));
+	if (!Element)
+	{
+		return;
+	}
+
+	FRigSocketElement* Socket = Cast<FRigSocketElement>(Element);
+	if (!Socket)
+	{
+		return;
+	}
+
+	TSharedPtr<FAssetDragDropOp> AssetDragDropOp = InDragDropEvent.GetOperationAs<FAssetDragDropOp>();
+	if (AssetDragDropOp.IsValid())
+	{
+		for (const FAssetData& AssetData : AssetDragDropOp->GetAssets())
+		{
+			UClass* AssetClass = AssetData.GetClass();
+			if (!AssetClass->IsChildOf(UControlRigBlueprint::StaticClass()))
+			{
+				continue;
+			}
+
+			if(UControlRigBlueprint* AssetBlueprint = Cast<UControlRigBlueprint>(AssetData.GetAsset()))
+			{
+				if (UModularRigController* Controller = Blueprint->GetModularRigController())
+				{
+					const FName ModuleName = Controller->GetSafeNewName(AssetBlueprint->RigModuleSettings.Identifier.Name);
+					const FString ModulePath = Controller->AddModule(ModuleName, AssetBlueprint->GetControlRigClass(), FString());
+					if(!ModulePath.IsEmpty())
+					{
+						FRigElementKey PrimaryConnectorKey;
+						TArray<FRigConnectorElement*> Connectors = Hierarchy->GetElementsOfType<FRigConnectorElement>();
+						for (FRigConnectorElement* Connector : Connectors)
+						{
+							if (Connector->Settings.Type == EConnectorType::Primary)
+							{
+								FString Path, Name;
+								Connector->GetName().Split(UModularRig::NamespaceSeparator, &Path, &Name, ESearchCase::CaseSensitive, ESearchDir::FromEnd);
+								if (Path == ModulePath)
+								{
+									PrimaryConnectorKey = Connector->GetKey();
+									break;
+								}
+							}
+						}
+						Controller->ConnectConnectorToElement(PrimaryConnectorKey, Socket->GetKey());
+					}
+				}
+			}
+		}
+	}
+}
+
 void FControlRigEditor::HandlePreviewMeshChanged(USkeletalMesh* InOldSkeletalMesh, USkeletalMesh* InNewSkeletalMesh)
 {
 	RebindToSkeletalMeshComponent();
@@ -3422,10 +3510,6 @@ void FControlRigEditor::OnHierarchyModified(ERigHierarchyNotification InNotif, U
 						RigBlueprint->TurnIntoControlRigModule();
 					}
 				}
-				else if(InElement->GetType() == ERigElementType::Socket)
-				{
-					SchematicGraph.AddNode(InElement->GetName());
-				}
 			}
 			// no break - fall through
 		}
diff --git a/Engine/Plugins/Animation/ControlRig/Source/ControlRigEditor/Private/Editor/ControlRigEditor.h b/Engine/Plugins/Animation/ControlRig/Source/ControlRigEditor/Private/Editor/ControlRigEditor.h
--- a/Engine/Plugins/Animation/ControlRig/Source/ControlRigEditor/Private/Editor/ControlRigEditor.h
+++ b/Engine/Plugins/Animation/ControlRig/Source/ControlRigEditor/Private/Editor/ControlRigEditor.h
@@ -26,6 +26,7 @@
 #include "RigVMHost.h"
 #include "SchematicGraphPanel/SSchematicGraphPanel.h"
 #include "Units/RigUnit.h"
+#include "ControlRigSchematicModel.h"
 
 class UControlRigBlueprint;
 class IPersonaToolkit;
@@ -246,6 +247,7 @@ private:
 	void HandleSchematicViewportCreated(const TSharedRef<class SSchematicGraphPanel>& InViewport);
 	void HandleUpdateSchematicNodes(SSchematicGraphPanel* InPanel, TSharedPtr<SSchematicGraphNode> InNode);
 	void HandleSchematicNodeClicked(SSchematicGraphPanel* InPanel, SSchematicGraphNode* InNode);
+	void HandleSchematicDrop(SSchematicGraphPanel* InPanel, SSchematicGraphNode* InNode, const FDragDropEvent& InDragDropEvent);
 	FVector2D ComputePersonaProjectedScreenPos(const FVector& InWorldPos);
 
 		/** Handle switching skeletal meshes */
@@ -304,7 +306,7 @@ protected:
 	UAnimPreviewInstance* PreviewInstance;
 
 	/** Model for the schematic views */
-	FSchematicGraph SchematicGraph;
+	FControlRigSchematicModel SchematicModel;
 
 	/** Delegate to deal with key down evens in the viewport / editor */
 	FPersonaViewportKeyDownDelegate OnKeyDownDelegate;
diff --git a/Engine/Plugins/Animation/ControlRig/Source/ControlRigEditor/Private/Editor/ControlRigSchematicModel.cpp b/Engine/Plugins/Animation/ControlRig/Source/ControlRigEditor/Private/Editor/ControlRigSchematicModel.cpp
new file mode 100644
--- /dev/null
+++ b/Engine/Plugins/Animation/ControlRig/Source/ControlRigEditor/Private/Editor/ControlRigSchematicModel.cpp
@@ -0,0 +1,86 @@
+// Copyright Epic Games, Inc. All Rights Reserved.
+
+#include "ControlRigSchematicModel.h"
+#include "ControlRig.h"
+
+void FControlRigSchematicModel::OnSetObjectBeingDebugged(UObject* InObject)
+{
+	if(ControlRigBeingDebuggedPtr.Get() == InObject)
+	{
+		return;
+	}
+
+	if(ControlRigBeingDebuggedPtr.IsValid())
+	{
+		if(UControlRig* ControlRigBeingDebugged = ControlRigBeingDebuggedPtr.Get())
+		{
+			if(!ControlRigBeingDebugged->HasAnyFlags(RF_BeginDestroyed))
+			{
+				ControlRigBeingDebugged->GetHierarchy()->OnModified().RemoveAll(this);
+			}
+		}
+	}
+
+	ControlRigBeingDebuggedPtr.Reset();
+	
+	if(UControlRig* ControlRig = Cast<UControlRig>(InObject))
+	{
+		ControlRigBeingDebuggedPtr = ControlRig;
+		if(URigHierarchy* Hierarchy = ControlRig->GetHierarchy())
+		{
+			Hierarchy->OnModified().RemoveAll(this);
+			Hierarchy->OnModified().AddRaw(this, &FControlRigSchematicModel::OnHierarchyModified);
+		}
+	}
+}
+
+void FControlRigSchematicModel::OnHierarchyModified(ERigHierarchyNotification InNotif, URigHierarchy* InHierarchy, const FRigBaseElement* InElement)
+{
+	switch (InNotif)
+	{
+		case ERigHierarchyNotification::ElementAdded:
+		{
+			if (InElement)
+			{
+				if (InElement->GetType() == ERigElementType::Socket)
+				{
+					SchematicGraph.AddNode(InElement->GetName());
+				}
+			}
+			break;
+		}
+		case ERigHierarchyNotification::ElementRenamed:
+		{
+			if (InElement)
+			{
+				if (InElement->GetType() == ERigElementType::Socket)
+				{
+					const FString OldNameStr = InHierarchy->GetPreviousName(InElement->GetKey()).ToString();
+					SchematicGraph.RenameNode(OldNameStr, InElement->GetName());
+				}
+			}
+			break;
+		}
+		case ERigHierarchyNotification::ElementRemoved:
+		{
+			if (InElement)
+			{
+				if (InElement->GetType() == ERigElementType::Socket)
+				{
+					SchematicGraph.RemoveNode(InElement->GetName());
+				}
+			}
+			break;
+		}
+		case ERigHierarchyNotification::HierarchyReset:
+		{
+			SchematicGraph.Reset();
+			TArray<FRigSocketElement*> Sockets = InHierarchy->GetElementsOfType<FRigSocketElement>();
+			for (FRigSocketElement* Socket : Sockets)
+			{
+				SchematicGraph.AddNode(Socket->GetFName().ToString());
+			}
+			break;
+		}
+	}
+}
diff --git a/Engine/Plugins/Animation/ControlRig/Source/ControlRigEditor/Private/Editor/ControlRigSchematicModel.h b/Engine/Plugins/Animation/ControlRig/Source/ControlRigEditor/Private/Editor/ControlRigSchematicModel.h
new file mode 100644
--- /dev/null
+++ b/Engine/Plugins/Animation/ControlRig/Source/ControlRigEditor/Private/Editor/ControlRigSchematicModel.h
@@ -0,0 +1,20 @@
+// Copyright Epic Games, Inc. All Rights Reserved.
+
+#pragma once
+
+#include "SchematicGraphPanel/SSchematicGraphPanel.h"
+#include "Rigs/RigHierarchyDefines.h"
+
+class UControlRig;
+class URigHierarchy;
+struct FRigBaseElement;
+
+/** Model for the schematic views */
+struct FControlRigSchematicModel 
+{
+	FSchematicGraph SchematicGraph;
+	TWeakObjectPtr<UControlRig> ControlRigBeingDebuggedPtr;
+
+	void OnSetObjectBeingDebugged(UObject* InObject);
+	void OnHierarchyModified(ERigHierarchyNotification InNotif, URigHierarchy* InHierarchy, const FRigBaseElement* InElement);
+};
\ No newline at end of file
diff --git a/Engine/Source/Developer/AnimationWidgets/Private/SchematicGraphPanel/SSchematicGraphPanel.cpp b/Engine/Source/Developer/AnimationWidgets/Private/SchematicGraphPanel/SSchematicGraphPanel.cpp
--- a/Engine/Source/Developer/AnimationWidgets/Private/SchematicGraphPanel/SSchematicGraphPanel.cpp
+++ b/Engine/Source/Developer/AnimationWidgets/Private/SchematicGraphPanel/SSchematicGraphPanel.cpp
@@ -2,11 +2,20 @@
 #if WITH_EDITOR
 
 #include "SchematicGraphPanel/SSchematicGraphPanel.h"
-#include "SGraphPanel.h"
+#include "Framework/Application/SlateApplication.h"
+#include "DragAndDrop/AssetDragDropOp.h"
 
 #define LOCTEXT_NAMESPACE "SSchematicGraphPanel"
 
 
+void FSchematicGraph::Reset()
+{
+	Nodes.Reset();
+	Links.Reset();
+
+	OnGraphReset().Broadcast();
+}
+
 bool FSchematicGraph::AddNode(const FString& InName)
 {
 	if (Nodes.ContainsByPredicate([InName](const FSchematicGraphNode* Node)
@@ -25,10 +34,43 @@ bool FSchematicGraph::AddNode(const FString& InName)
 	return true;
 }
 
+bool FSchematicGraph::RenameNode(const FString& InOldName, const FString& InNewName)
+{
+	FSchematicGraphNode** FoundNode = Nodes.FindByPredicate([InOldName](const FSchematicGraphNode* Node)
+	{
+		return Node->Name == InOldName;
+	});
+
+	if (FoundNode)
+	{
+		(*FoundNode)->Name = InNewName;
+		return true;
+	}
+	
+	return false;
+}
+
+bool FSchematicGraph::RemoveNode(const FString& InName)
+{
+	FSchematicGraphNode** FoundNode = Nodes.FindByPredicate([InName](const FSchematicGraphNode* Node)
+	{
+		return Node->Name == InName;
+	});
+	if (FoundNode)
+	{
+		OnNodeRemovedDelegate.Broadcast(*FoundNode);
+		Nodes.Remove(*FoundNode);
+		FMemory::Free(*FoundNode);
+		return true;
+	}
+	return false;
+}
+
 void SSchematicGraphNode::Construct(const FArguments& InArgs)
 {
 	NodeData = InArgs._NodeData;
 	OnClickedDelegate = InArgs._OnClicked;
+	OnDropDelegate = InArgs._OnDrop;
 	TEasingAttributeInterpolator<FVector2d>::FSettings Vector2DInterpSettings(EEasingInterpolatorType::BounceEaseInOut, 0.2f);
 	Position = TAnimatedAttribute<FVector2d>::Create(Vector2DInterpSettings, FVector2d::ZeroVector);
 	Size = TAnimatedAttribute<FVector2d>::Create(Vector2DInterpSettings, OriginalSize);
@@ -54,17 +96,13 @@ int32 SSchematicGraphNode::OnPaint(const FPaintArgs& Args, const FGeometry& Allo
 	//FVector2d TotalOffset = CenterOffset + SizeOffset;
 	FVector2d TotalOffset = SizeOffset;
 
-	FLinearColor Color = NodeData->bIsSelected ?
-		FSlateColor(EStyleColor::AccentRed).GetSpecifiedColor()
-	: FLinearColor::White;
-
 	FSlateDrawElement::MakeBox(
 				OutDrawElements,
 				NewLayerId,
 				AllottedGeometry.ToPaintGeometry(CurSize, FSlateLayoutTransform(TotalOffset)),
 				&Brush,
 				ESlateDrawEffect::None,
-				Color
+				Brush.TintColor.GetSpecifiedColor()
 				);
 	
 	return NewLayerId;
@@ -95,6 +133,13 @@ void SSchematicGraphNode::OnDragLeave(const FDragDropEvent& DragDropEvent)
 	Scale->Set(1.f);
 }
 
+FReply SSchematicGraphNode::OnDrop(const FGeometry& MyGeometry, const FDragDropEvent& DragDropEvent)
+{
+	SNode::OnDrop(MyGeometry, DragDropEvent);
+	OnDropDelegate.ExecuteIfBound(this, DragDropEvent);
+	return FReply::Unhandled();
+}
+
 FReply SSchematicGraphNode::OnMouseButtonDown(const FGeometry& MyGeometry, const FPointerEvent& MouseEvent)
 {
 	SNode::OnMouseButtonDown(MyGeometry, MouseEvent);
@@ -115,12 +160,16 @@ void SSchematicGraphPanel::SetSchematicGraph(FSchematicGraph* InGraphData)
 	if (GraphData)
 	{
 		GraphData->OnNodeAdded().RemoveAll(this);
+		GraphData->OnNodeRemoved().RemoveAll(this);
+		GraphData->OnGraphReset().RemoveAll(this);
 	}
 	
 	GraphData = InGraphData;
 	if (GraphData)
 	{
 		GraphData->OnNodeAdded().AddSP(this, &SSchematicGraphPanel::AddNode);
+		GraphData->OnNodeRemoved().AddSP(this, &SSchematicGraphPanel::RemoveNode);
+		GraphData->OnGraphReset().AddSP(this, &SSchematicGraphPanel::RebuildPanel);
 	}
 }
 
@@ -130,6 +179,10 @@ void SSchematicGraphPanel::Construct(const FArguments& InArgs)
 	bIsOverlay = InArgs._IsOverlay;
 	UpdateNodeWidgetDelegate = InArgs._OnUpdateNodeWidget;
 	OnNodeClickedDelegate = InArgs._OnNodeClicked;
+	OnDropDelegate = InArgs._OnDrop;
+
+	TEasingAttributeInterpolator<float>::FSettings FloatInterpSettings(EEasingInterpolatorType::BounceEaseInOut, 0.2f);
+	FadeBackgroundAlpha = TAnimatedAttribute<float>::Create(FloatInterpSettings, 0.f);
 	
 	SNodePanel::Construct();
 
@@ -137,6 +190,8 @@ void SSchematicGraphPanel::Construct(const FArguments& InArgs)
 	if (GraphData)
 	{
 		GraphData->OnNodeAdded().AddSP(this, &SSchematicGraphPanel::AddNode);
+		GraphData->OnNodeRemoved().AddSP(this, &SSchematicGraphPanel::RemoveNode);
+		GraphData->OnGraphReset().AddSP(this, &SSchematicGraphPanel::RebuildPanel);
 	}
 }
 
@@ -157,19 +212,61 @@ void SSchematicGraphPanel::AddNode(FSchematicGraphNode* NodeToAdd)
 {
 	TSharedRef<SSchematicGraphNode> NewNode = SNew(SSchematicGraphNode)
 														.OnClicked_Raw(this, &SSchematicGraphPanel::OnNodeClicked)
+														.OnDrop_Raw(this, &SSchematicGraphPanel::OnDropEvent)
 														.NodeData(NodeToAdd);
 	SNodePanel::AddGraphNode(NewNode);
 }
 
+void SSchematicGraphPanel::RemoveNode(FSchematicGraphNode* InNodeToRemove)
+{
+	for (int32 Iter = 0; Iter != Children.Num(); ++Iter)
+	{
+		TSharedRef<SSchematicGraphNode> Child = GetChild(Iter);
+		if (Child->NodeData == InNodeToRemove)
+		{
+			Children.RemoveAt(Iter);
+			break;
+		}
+	}
+	for (int32 Iter = 0; Iter != VisibleChildren.Num(); ++Iter)
+	{
+		TSharedRef<SSchematicGraphNode> Child = StaticCastSharedRef<SSchematicGraphNode>(VisibleChildren[Iter]);
+		if (Child->NodeData == InNodeToRemove)
+		{
+			VisibleChildren.RemoveAt(Iter);
+			break;
+		}
+	}
+}
+
 int32 SSchematicGraphPanel::OnPaint(const FPaintArgs& Args, const FGeometry& AllottedGeometry, const FSlateRect& MyCullingRect, FSlateWindowElementList& OutDrawElements, int32 LayerId, const FWidgetStyle& InWidgetStyle, bool bParentEnabled) const
 {
+	const int32 NodeLayerId = LayerId + 1;
+    int32 MaxLayerId = NodeLayerId;
+    	
 	if (!bIsOverlay)
 	{
 		const FSlateBrush* DefaultBackground = FAppStyle::GetBrush(TEXT("Graph.Panel.SolidBackground"));
 		PaintBackgroundAsLines(DefaultBackground, AllottedGeometry, MyCullingRect, OutDrawElements, LayerId);
+		MaxLayerId++;
 	}
-	const int32 NodeLayerId = LayerId + 1;
-	int32 MaxLayerId = NodeLayerId;
+
+	if (FadeBackgroundAlpha->Get() > 0.f)
+	{
+		const FSlateBrush* Brush = FAppStyle::GetBrush(TEXT("Graph.Panel.SolidBackground"));
+		FLinearColor TransparentGrey = FLinearColor(0.5, 0.5, 0.5, FadeBackgroundAlpha->Get());
+		FSlateDrawElement::MakeBox(
+				OutDrawElements,
+				MaxLayerId,
+				AllottedGeometry.ToPaintGeometry(),
+				Brush,
+				ESlateDrawEffect::None,
+				TransparentGrey
+				);
+		MaxLayerId++;
+	}
+	
+	
 
 	if (!GraphData)
 	{
@@ -223,6 +320,18 @@ TStatId SSchematicGraphPanel::GetStatId() const
 
 void SSchematicGraphPanel::Tick(float DeltaTime)
 {
+	bool bFadeBackgroundBlanket = false;
+	FSlateApplication& Application = FSlateApplication::Get();
+	if (Application.IsDragDropping())
+	{
+		TSharedPtr<FDragDropOperation> DragDropOp = Application.GetDragDroppingContent();
+		if (DragDropOp.IsValid())
+		{
+			bFadeBackgroundBlanket = true;
+		}
+	}
+	SetFadeBackground(bFadeBackgroundBlanket);
+	
 	for (int32 i=0; i<Children.Num(); ++i)
 	{
 		TSharedRef<SSchematicGraphNode> SNode = GetChild(i);
@@ -232,6 +341,7 @@ void SSchematicGraphPanel::Tick(float DeltaTime)
 
 void SSchematicGraphPanel::Tick(const FGeometry& AllottedGeometry, const double InCurrentTime, const float InDeltaTime)
 {
+	SNodePanel::Tick(AllottedGeometry, InCurrentTime, InDeltaTime);
 	Tick(InDeltaTime);
 }
 
@@ -240,6 +350,11 @@ void SSchematicGraphPanel::OnNodeClicked(SSchematicGraphNode* Node)
 	OnNodeClickedDelegate.ExecuteIfBound(this, Node);
 }
 
+void SSchematicGraphPanel::OnDropEvent(SSchematicGraphNode* Node, const FDragDropEvent& InDragDropEvent)
+{
+	OnDropDelegate.ExecuteIfBound(this, Node, InDragDropEvent);
+}
+
 
 #undef LOCTEXT_NAMESPACE
 
diff --git a/Engine/Source/Developer/AnimationWidgets/Public/SchematicGraphPanel/SSchematicGraphPanel.h b/Engine/Source/Developer/AnimationWidgets/Public/SchematicGraphPanel/SSchematicGraphPanel.h
--- a/Engine/Source/Developer/AnimationWidgets/Public/SchematicGraphPanel/SSchematicGraphPanel.h
+++ b/Engine/Source/Developer/AnimationWidgets/Public/SchematicGraphPanel/SSchematicGraphPanel.h
@@ -14,6 +14,8 @@ struct ANIMATIONWIDGETS_API FSchematicGraphNode
 };
 
 DECLARE_EVENT_OneParam(FSchematicGraph, FOnNodeAdded, FSchematicGraphNode*);
+DECLARE_EVENT_OneParam(FSchematicGraph, FOnNodeRemoved, FSchematicGraphNode*);
+DECLARE_EVENT(FSchematicGraph, FOnGraphReset);
 
 class ANIMATIONWIDGETS_API FSchematicGraph
 {
@@ -21,20 +23,29 @@ public:
 	TArray<FSchematicGraphNode*> Nodes;
 	TMap<FSchematicGraphNode*, TArray<FSchematicGraphNode*>> Links;
 
+	void Reset();
 	bool AddNode(const FString& InName);
+	bool RenameNode(const FString& InOldName, const FString& InNewName);
+	bool RemoveNode(const FString& InName);
 
 	FOnNodeAdded OnNodeAddedDelegate;
 	FOnNodeAdded& OnNodeAdded() { return OnNodeAddedDelegate; }
+	FOnNodeRemoved OnNodeRemovedDelegate;
+	FOnNodeRemoved& OnNodeRemoved() { return OnNodeRemovedDelegate; }
+	FOnGraphReset OnGraphResetDelegate;
+	FOnGraphReset& OnGraphReset() { return OnGraphResetDelegate; }
 };
 
 class ANIMATIONWIDGETS_API SSchematicGraphNode : public SNodePanel::SNode
 {
 public:
 	DECLARE_DELEGATE_OneParam(FOnClicked, SSchematicGraphNode*);
+	DECLARE_DELEGATE_TwoParams(FOnDrop, SSchematicGraphNode*, const FDragDropEvent&);
 
 	SLATE_BEGIN_ARGS(SSchematicGraphNode){}
 	SLATE_ARGUMENT(FSchematicGraphNode*, NodeData)
 	SLATE_EVENT(FOnClicked, OnClicked)
+	SLATE_EVENT(FOnDrop, OnDrop)
 	SLATE_END_ARGS()
 
 	void Construct(const FArguments& InArgs);
@@ -44,9 +55,10 @@ public:
 	
 	virtual void OnMouseEnter(const FGeometry& MyGeometry, const FPointerEvent& MouseEvent) override;
 	virtual void OnMouseLeave(const FPointerEvent& MouseEvent) override;
+	virtual FReply OnMouseButtonDown(const FGeometry& MyGeometry, const FPointerEvent& MouseEvent) override;
 	virtual FReply OnDragOver(const FGeometry& MyGeometry, const FDragDropEvent& DragDropEvent) override;
 	virtual void OnDragLeave(const FDragDropEvent& DragDropEvent) override;
-	virtual FReply OnMouseButtonDown(const FGeometry& MyGeometry, const FPointerEvent& MouseEvent) override;
+	virtual FReply OnDrop(const FGeometry& MyGeometry, const FDragDropEvent& DragDropEvent) override;
 
 	virtual FVector2d GetPosition() const override;
 
@@ -62,6 +74,7 @@ public:
 	float ScaledDown = 0.75;
 
 	FOnClicked OnClickedDelegate;
+	FOnDrop OnDropDelegate;
 	FSchematicGraphNode* NodeData;
 	TSharedPtr<TAnimatedAttribute<FVector2d>> Position;
 	TSharedPtr<TAnimatedAttribute<FVector2d>> Size;
@@ -77,12 +90,14 @@ public:
 	
 	DECLARE_DELEGATE_TwoParams(FUpdateNodeWidget, SSchematicGraphPanel*, TSharedPtr<SSchematicGraphNode>);
 	DECLARE_DELEGATE_TwoParams(FOnNodeClicked, SSchematicGraphPanel*, SSchematicGraphNode*);
+	DECLARE_DELEGATE_ThreeParams(FOnDrop, SSchematicGraphPanel*, SSchematicGraphNode*, const FDragDropEvent&);
 	
 	SLATE_BEGIN_ARGS(SSchematicGraphPanel) {}
 	SLATE_ARGUMENT(bool, IsOverlay)
 	SLATE_ARGUMENT(FSchematicGraph*, GraphData)
 	SLATE_EVENT(FUpdateNodeWidget, OnUpdateNodeWidget)
 	SLATE_EVENT(FOnNodeClicked, OnNodeClicked)
+	SLATE_EVENT(FOnDrop, OnDrop)
 	SLATE_END_ARGS()
 
 	~SSchematicGraphPanel(){}
@@ -92,11 +107,12 @@ public:
 
 	void RebuildPanel();
 	void AddNode(FSchematicGraphNode* NodeToAdd);
+	void RemoveNode(FSchematicGraphNode* NodeToRemove);
 
 	// SNodePanel interface
 	virtual int32 OnPaint(const FPaintArgs& Args, const FGeometry& AllottedGeometry, const FSlateRect& MyCullingRect, FSlateWindowElementList& OutDrawElements, int32 LayerId, const FWidgetStyle& InWidgetStyle, bool bParentEnabled) const override;
 	// End of SNodePanel interface
-
+	
 	TSharedRef<SSchematicGraphNode> GetChild(int32 ChildIndex) const;
 
 	// FTickableEditorObject Interface
@@ -107,11 +123,16 @@ public:
 	// End of FTickableEditorObject interface
 
 	void OnNodeClicked(SSchematicGraphNode* Node);
+	void OnDropEvent(SSchematicGraphNode* Node, const FDragDropEvent& InDragDropEvent);
+
+	void SetFadeBackground(bool bInFade) { FadeBackgroundAlpha->Set(bInFade ? 0.5f : 0.f); }
 
 	bool bIsOverlay;
+	TSharedPtr<TAnimatedAttribute<float>> FadeBackgroundAlpha;
 	FSchematicGraph* GraphData;
 	FUpdateNodeWidget UpdateNodeWidgetDelegate;
 	FOnNodeClicked OnNodeClickedDelegate;
+	FOnDrop OnDropDelegate;
 };
 
 #endif
\ No newline at end of file
