Control Rig: Drag and drop nodes in the schematic panel
#rnx

diff --git a/Engine/Plugins/Animation/ControlRig/Source/ControlRigEditor/Private/Editor/ControlRigSchematicModel.cpp b/Engine/Plugins/Animation/ControlRig/Source/ControlRigEditor/Private/Editor/ControlRigSchematicModel.cpp
--- a/Engine/Plugins/Animation/ControlRig/Source/ControlRigEditor/Private/Editor/ControlRigSchematicModel.cpp
+++ b/Engine/Plugins/Animation/ControlRig/Source/ControlRigEditor/Private/Editor/ControlRigSchematicModel.cpp
@@ -7,6 +7,7 @@
 #include "ControlRigEditorStyle.h"
 #include "ModularRig.h"
 #include "DragAndDrop/AssetDragDropOp.h"
+#include "Rigs/RigHierarchyController.h"
 
 FControlRigSchematicModel::~FControlRigSchematicModel()
 {
@@ -156,6 +157,18 @@ void FControlRigSchematicModel::HandleSchematicNodeClicked(SSchematicGraphPanel*
 	{
 		Node->bIsSelected = Node == InNode->NodeData;
 	}
+
+	if (ControlRigBeingDebuggedPtr.IsValid())
+	{
+		if (URigHierarchy* Hierarchy = ControlRigBeingDebuggedPtr->GetHierarchy())
+		{
+			if (URigHierarchyController* Controller = Hierarchy->GetController())
+			{
+				TArray<FRigElementKey> Selection = {FRigElementKey(InNode->NodeData->Name)};
+				Controller->SetSelection(Selection);
+			}
+		}
+	}
 }
 
 void FControlRigSchematicModel::HandleUpdateSchematicNodes(SSchematicGraphPanel* InPanel, TSharedPtr<SSchematicGraphNode> InNode)
@@ -251,19 +264,14 @@ void FControlRigSchematicModel::HandleSchematicDrop(SSchematicGraphPanel* InPane
 		return;
 	}
 	
-	FRigBaseElement* Element = Hierarchy->Find(FRigElementKey(InNode->NodeData->Name));
-	if (!Element)
-	{
-		return;
-	}
-
-	FRigSocketElement* Socket = Cast<FRigSocketElement>(Element);
-	if (!Socket)
+	FRigBaseElement* Target = Hierarchy->Find(FRigElementKey(InNode->NodeData->Name));
+	if (!Target)
 	{
 		return;
 	}
 
 	TSharedPtr<FAssetDragDropOp> AssetDragDropOp = InDragDropEvent.GetOperationAs<FAssetDragDropOp>();
+	TSharedPtr<FSchematicGraphNodeDragDropOp> SchematicDragDropOp = InDragDropEvent.GetOperationAs<FSchematicGraphNodeDragDropOp>();
 	if (AssetDragDropOp.IsValid())
 	{
 		for (const FAssetData& AssetData : AssetDragDropOp->GetAssets())
@@ -297,7 +305,25 @@ void FControlRigSchematicModel::HandleSchematicDrop(SSchematicGraphPanel* InPane
 								}
 							}
 						}
-						Controller->ConnectConnectorToElement(PrimaryConnectorKey, Socket->GetKey());
+						Controller->ConnectConnectorToElement(PrimaryConnectorKey, Target->GetKey());
+					}
+				}
+			}
+		}
+	}
+	else if(SchematicDragDropOp.IsValid())
+	{
+		if (UModularRigController* Controller = ControlRigBlueprint->GetModularRigController())
+		{
+			TArray<FString> Sources = SchematicDragDropOp->GetElements();
+			for (FString& SourceName : Sources)
+			{
+				FRigElementKey Key(SourceName);
+				if (FRigBaseElement* SourceElement = Hierarchy->Find(SourceName))
+				{
+					if (FRigConnectorElement* Connector = Cast<FRigConnectorElement>(SourceElement))
+					{
+						Controller->ConnectConnectorToElement(Key, Target->GetKey());
 					}
 				}
 			}
diff --git a/Engine/Source/Developer/AnimationWidgets/Private/SchematicGraphPanel/SSchematicGraphPanel.cpp b/Engine/Source/Developer/AnimationWidgets/Private/SchematicGraphPanel/SSchematicGraphPanel.cpp
--- a/Engine/Source/Developer/AnimationWidgets/Private/SchematicGraphPanel/SSchematicGraphPanel.cpp
+++ b/Engine/Source/Developer/AnimationWidgets/Private/SchematicGraphPanel/SSchematicGraphPanel.cpp
@@ -80,6 +80,35 @@ bool FSchematicGraph::RemoveNode(const FString& InName)
 	return false;
 }
 
+TSharedRef<FSchematicGraphNodeDragDropOp> FSchematicGraphNodeDragDropOp::New(const TArray<FString>& InElements)
+{
+	TSharedRef<FSchematicGraphNodeDragDropOp> Operation = MakeShared<FSchematicGraphNodeDragDropOp>();
+	Operation->Elements = InElements;
+	Operation->Construct();
+	return Operation;
+}
+
+TSharedPtr<SWidget> FSchematicGraphNodeDragDropOp::GetDefaultDecorator() const
+{
+	return SNew(SBorder)
+		.Visibility(EVisibility::Visible)
+		.BorderImage(FAppStyle::GetBrush("Menu.Background"))
+		[
+			SNew(STextBlock)
+			.Text(FText::FromString(GetJoinedElementNames()))
+		];
+}
+
+FString FSchematicGraphNodeDragDropOp::GetJoinedElementNames() const
+{
+	TArray<FString> ElementNameStrings;
+	for (const FString& Element: Elements)
+	{
+		ElementNameStrings.Add(Element);
+	}
+	return FString::Join(ElementNameStrings, TEXT(","));
+}
+
 void SSchematicGraphNode::Construct(const FArguments& InArgs)
 {
 	NodeData = InArgs._NodeData;
@@ -105,10 +134,8 @@ int32 SSchematicGraphNode::OnPaint(const FPaintArgs& Args, const FGeometry& Allo
 	int32 NewLayerId = SNodePanel::SNode::OnPaint(Args, AllottedGeometry, MyCullingRect, OutDrawElements, LayerId, InWidgetStyle, bParentEnabled);
 	NewLayerId++;
 
-	//FVector2d CenterOffset = OriginalSize*-0.5;
 	FVector2d CurSize = Size->Get() * Scale->Get();
 	FVector2d SizeOffset = (CurSize-OriginalSize)*-0.5;
-	//FVector2d TotalOffset = CenterOffset + SizeOffset;
 	FVector2d TotalOffset = SizeOffset;
 
 	FSlateDrawElement::MakeBox(
@@ -155,6 +182,19 @@ FReply SSchematicGraphNode::OnDrop(const FGeometry& MyGeometry, const FDragDropE
 	return FReply::Unhandled();
 }
 
+FReply SSchematicGraphNode::OnDragDetected(const FGeometry& MyGeometry, const FPointerEvent& MouseEvent)
+{
+	TArray<FString> DraggedElements = {NodeData->Name};
+	if (MouseEvent.IsMouseButtonDown(EKeys::LeftMouseButton) && DraggedElements.Num() > 0)
+	{
+		bIsBeingDragged = true;
+		TSharedRef<FSchematicGraphNodeDragDropOp> DragDropOp = FSchematicGraphNodeDragDropOp::New(MoveTemp(DraggedElements));
+		return FReply::Handled().BeginDragDrop(DragDropOp);
+	}
+	
+	return FReply::Unhandled();
+}
+
 FReply SSchematicGraphNode::OnMouseButtonDown(const FGeometry& MyGeometry, const FPointerEvent& MouseEvent)
 {
 	SNode::OnMouseButtonDown(MyGeometry, MouseEvent);
@@ -162,7 +202,7 @@ FReply SSchematicGraphNode::OnMouseButtonDown(const FGeometry& MyGeometry, const
 	{
 		OnClickedDelegate.ExecuteIfBound(this);
 	}
-	return FReply::Handled();
+	return FReply::Handled().DetectDrag(SharedThis(this), EKeys::LeftMouseButton);
 }
 
 FVector2d SSchematicGraphNode::GetPosition() const
@@ -170,6 +210,13 @@ FVector2d SSchematicGraphNode::GetPosition() const
 	return Position->Get() - (OriginalSize*0.5);
 }
 
+void SSchematicGraphNode::SetPosition(const FVector2d& InPosition, bool bImmediately)
+{
+	bImmediately ?
+		Position->SetValueAndStop(FVector2d(InPosition))
+		: Position->Set(InPosition);
+}
+
 void SSchematicGraphPanel::SetSchematicGraph(FSchematicGraph* InGraphData)
 {
 	if (GraphData)
@@ -238,6 +285,12 @@ void SSchematicGraphPanel::AddNode(FSchematicGraphNode* NodeToAdd)
 														.OnDrop_Raw(this, &SSchematicGraphPanel::OnDropEvent)
 														.NodeData(NodeToAdd);
 	SNodePanel::AddGraphNode(NewNode);
+
+	// Update the position/size/brush immediately without animation
+	{
+		TGuardValue<bool> AnimateNodeGuard(bAnimatePosition, false);
+		Tick(0.f);
+	}
 }
 
 void SSchematicGraphPanel::RemoveNode(FSchematicGraphNode* InNodeToRemove)
@@ -365,17 +418,21 @@ TStatId SSchematicGraphPanel::GetStatId() const
 
 void SSchematicGraphPanel::Tick(float DeltaTime)
 {
-	bool bFadeBackgroundBlanket = false;
+	bool bIsDragging = false;
 	FSlateApplication& Application = FSlateApplication::Get();
 	if (Application.IsDragDropping())
 	{
 		TSharedPtr<FDragDropOperation> DragDropOp = Application.GetDragDroppingContent();
 		if (DragDropOp.IsValid())
 		{
-			bFadeBackgroundBlanket = true;
+			bIsDragging = true;
 		}
 	}
-	SetFadeBackground(bFadeBackgroundBlanket);
+	else
+	{
+		bIsDragging = false;
+	}
+	SetFadeBackground(bIsDragging);
 
 	uint16 TopRightNodes = 0;
 	uint16 TopLeftNodes = 0;
@@ -386,7 +443,9 @@ void SSchematicGraphPanel::Tick(float DeltaTime)
 	{
 		TSharedRef<SSchematicGraphNode> SNode = GetChild(i);
 		UpdateNodeWidgetDelegate.ExecuteIfBound(this, SNode.ToSharedPtr());
-		
+
+		FVector2d NewPosition = SNode->GetPosition() + (SNode->OriginalSize*0.5);
+		bool bImmediatePosition = !bAnimatePosition;
 		switch(SNode->NodeData->Placement)
 		{
 			case ESchematicGraphNodePlacement::Panel:
@@ -395,29 +454,49 @@ void SSchematicGraphPanel::Tick(float DeltaTime)
 			}
 			case ESchematicGraphNodePlacement::TopLeft:
 			{
-				SNode->SetPosition(FVector2d(PaddingLeft + SNode->OriginalSize.X*0.5, PaddingTop + SNode->OriginalSize.Y*0.5 + ((SNode->OriginalSize.Y + PaddingInterNode) * TopLeftNodes)), true);
+				NewPosition = FVector2d(PaddingLeft + SNode->OriginalSize.X*0.5, PaddingTop + SNode->OriginalSize.Y*0.5 + ((SNode->OriginalSize.Y + PaddingInterNode) * TopLeftNodes));
 				TopLeftNodes++;
 				break;
 			}
 			case ESchematicGraphNodePlacement::TopRight:
 			{
-				SNode->SetPosition(FVector2d(Size.X - PaddingRight - (SNode->OriginalSize.X * 0.5), PaddingTop + SNode->OriginalSize.Y*0.5 + ((SNode->OriginalSize.Y + PaddingInterNode) * TopRightNodes)), true);
+				NewPosition = FVector2d(Size.X - PaddingRight - (SNode->OriginalSize.X * 0.5), PaddingTop + SNode->OriginalSize.Y*0.5 + ((SNode->OriginalSize.Y + PaddingInterNode) * TopRightNodes));
 				TopRightNodes++;
 				break;
 			}
 			case ESchematicGraphNodePlacement::BottomLeft:
 			{
-				SNode->SetPosition(FVector2d(PaddingLeft + SNode->OriginalSize.X*0.5, Size.Y - PaddingBottom - (SNode->OriginalSize.Y * 0.5) - ((SNode->OriginalSize.Y + PaddingInterNode)  * BottomLeftNodes)), true);
+				NewPosition = FVector2d(PaddingLeft + SNode->OriginalSize.X*0.5, Size.Y - PaddingBottom - (SNode->OriginalSize.Y * 0.5) - ((SNode->OriginalSize.Y + PaddingInterNode)  * BottomLeftNodes));
 				BottomLeftNodes++;
 				break;
 			}
 			case ESchematicGraphNodePlacement::BottomRight:
 			{
-				SNode->SetPosition(FVector2d(Size.X - PaddingRight - (SNode->OriginalSize.X * 0.5), Size.Y - PaddingBottom - (SNode->OriginalSize.Y * 0.5) - ((SNode->OriginalSize.Y + PaddingInterNode)  * BottomRightNodes)), true);
+				NewPosition = FVector2d(Size.X - PaddingRight - (SNode->OriginalSize.X * 0.5), Size.Y - PaddingBottom - (SNode->OriginalSize.Y * 0.5) - ((SNode->OriginalSize.Y + PaddingInterNode)  * BottomRightNodes));
 				BottomRightNodes++;
 				break;
 			}
 		}
+
+		if (SNode->bIsBeingDragged)
+		{
+			if (bIsDragging)
+			{
+				const FVector2f AbsoluteMousePosition = FSlateApplication::Get().GetCursorPos();
+				const FGeometry& Geometry = GetTickSpaceGeometry();
+
+				const FVector2d LocalMousePosition =  (AbsoluteMousePosition - Geometry.GetAbsolutePosition()) / Geometry.GetAccumulatedLayoutTransform().GetScale();
+				NewPosition = LocalMousePosition;
+				bImmediatePosition = true;
+				
+			}
+			else
+			{
+				SNode->bIsBeingDragged = false;
+			}
+		}
+
+		SNode->SetPosition(NewPosition, bImmediatePosition);
 	}
 }
 
diff --git a/Engine/Source/Developer/AnimationWidgets/Public/SchematicGraphPanel/SSchematicGraphPanel.h b/Engine/Source/Developer/AnimationWidgets/Public/SchematicGraphPanel/SSchematicGraphPanel.h
--- a/Engine/Source/Developer/AnimationWidgets/Public/SchematicGraphPanel/SSchematicGraphPanel.h
+++ b/Engine/Source/Developer/AnimationWidgets/Public/SchematicGraphPanel/SSchematicGraphPanel.h
@@ -49,6 +49,35 @@ public:
 	FOnGraphReset& OnGraphReset() { return OnGraphResetDelegate; }
 };
 
+class FSchematicGraphNodeDragDropOp : public FDragDropOperation
+{
+public:
+	DRAG_DROP_OPERATOR_TYPE(FSchematicGraphNodeDragDropOp, FDragDropOperation)
+
+	static TSharedRef<FSchematicGraphNodeDragDropOp> New(const TArray<FString>& InElements);
+
+	virtual TSharedPtr<SWidget> GetDefaultDecorator() const override;
+
+	/** @return true if this drag operation contains property paths */
+	bool HasElements() const
+	{
+		return Elements.Num() > 0;
+	}
+
+	/** @return The property paths from this drag operation */
+	const TArray<FString>& GetElements() const
+	{
+		return Elements;
+	}
+
+	FString GetJoinedElementNames() const;
+
+private:
+
+	/** Data for the property paths this item represents */
+	TArray<FString> Elements;
+};
+
 class ANIMATIONWIDGETS_API SSchematicGraphNode : public SNodePanel::SNode
 {
 public:
@@ -72,16 +101,13 @@ public:
 	virtual FReply OnDragOver(const FGeometry& MyGeometry, const FDragDropEvent& DragDropEvent) override;
 	virtual void OnDragLeave(const FDragDropEvent& DragDropEvent) override;
 	virtual FReply OnDrop(const FGeometry& MyGeometry, const FDragDropEvent& DragDropEvent) override;
+	virtual FReply OnDragDetected(const FGeometry& MyGeometry, const FPointerEvent& MouseEvent) override;
 
 	virtual FVector2d GetPosition() const override;
 
-	void SetPosition(const FVector2d& InPosition, bool bImmediately = false)
-	{
-		bImmediately ?
-			Position->SetValueAndStop(FVector2d(InPosition))
-		: Position->Set(InPosition);
-	}
-	
+	void SetPosition(const FVector2d& InPosition, bool bImmediately = false);
+
+	bool bIsBeingDragged = false;
 	FVector2d OriginalSize = FVector2d(50.0,50.0);
 	float ScaledUp = 1.25;
 	float ScaledDown = 0.75;
@@ -151,6 +177,7 @@ public:
 
 	void SetFadeBackground(bool bInFade) { FadeBackgroundAlpha->Set(bInFade ? 0.5f : 0.f); }
 
+	bool bAnimatePosition = true;
 	bool bIsOverlay;
 	int32 PaddingLeft = 0;
 	int32 PaddingRight = 0;
