Control Rig: Show possible matches during dragging of connector
#rnx

diff --git a/Engine/Plugins/Animation/ControlRig/Source/ControlRig/Private/ModularRigController.cpp b/Engine/Plugins/Animation/ControlRig/Source/ControlRig/Private/ModularRigController.cpp
--- a/Engine/Plugins/Animation/ControlRig/Source/ControlRig/Private/ModularRigController.cpp
+++ b/Engine/Plugins/Animation/ControlRig/Source/ControlRig/Private/ModularRigController.cpp
@@ -161,6 +161,18 @@ bool UModularRigController::ConnectConnectorToElement(const FRigElementKey& InCo
 		return false;
 	}
 
+	if (!InTargetKey.IsValid())
+	{
+		UE_LOG(LogControlRig, Error, TEXT("Invalid target %s in class %s"), *InTargetKey.ToString(), *Module->Class->GetPathName());
+		return false;
+	}
+
+	if (InTargetKey == InConnectorKey)
+	{
+		UE_LOG(LogControlRig, Error, TEXT("Cannot resolve connector %s to itself in class %s"), *InTargetKey.ToString(), *Module->Class->GetPathName());
+		return false;
+	}
+
 	/*
 	FText ErrorMessage;
 	if (!CanConnectConnectorToElement(*ModuleConnector, InTargetKey, ErrorMessage))
diff --git a/Engine/Plugins/Animation/ControlRig/Source/ControlRigEditor/Private/Editor/ControlRigEditor.cpp b/Engine/Plugins/Animation/ControlRig/Source/ControlRigEditor/Private/Editor/ControlRigEditor.cpp
--- a/Engine/Plugins/Animation/ControlRig/Source/ControlRigEditor/Private/Editor/ControlRigEditor.cpp
+++ b/Engine/Plugins/Animation/ControlRig/Source/ControlRigEditor/Private/Editor/ControlRigEditor.cpp
@@ -2573,6 +2573,8 @@ void FControlRigEditor::HandleSchematicViewportCreated(const TSharedRef<SSchemat
 {
 	InViewport->UpdateNodeWidgetDelegate.BindRaw(&SchematicModel, &FControlRigSchematicModel::HandleUpdateSchematicNodes);
 	InViewport->OnNodeClickedDelegate.BindRaw(&SchematicModel, &FControlRigSchematicModel::HandleSchematicNodeClicked);
+	InViewport->OnBeginDragDelegate.BindRaw(&SchematicModel, &FControlRigSchematicModel::HandleSchematicBeginDrag);
+	InViewport->OnEndDragDelegate.BindRaw(&SchematicModel, &FControlRigSchematicModel::HandleSchematicEndDrag);
 	InViewport->OnDropDelegate.BindRaw(&SchematicModel, &FControlRigSchematicModel::HandleSchematicDrop);
 }
 
diff --git a/Engine/Plugins/Animation/ControlRig/Source/ControlRigEditor/Private/Editor/ControlRigSchematicModel.cpp b/Engine/Plugins/Animation/ControlRig/Source/ControlRigEditor/Private/Editor/ControlRigSchematicModel.cpp
--- a/Engine/Plugins/Animation/ControlRig/Source/ControlRigEditor/Private/Editor/ControlRigSchematicModel.cpp
+++ b/Engine/Plugins/Animation/ControlRig/Source/ControlRigEditor/Private/Editor/ControlRigSchematicModel.cpp
@@ -6,6 +6,7 @@
 #include "ControlRigEditor.h"
 #include "ControlRigEditorStyle.h"
 #include "ModularRig.h"
+#include "ModularRigRuleManager.h"
 #include "DragAndDrop/AssetDragDropOp.h"
 #include "Rigs/RigHierarchyController.h"
 
@@ -171,6 +172,106 @@ void FControlRigSchematicModel::HandleSchematicNodeClicked(SSchematicGraphPanel*
 	}
 }
 
+void FControlRigSchematicModel::HandleSchematicBeginDrag(SSchematicGraphPanel* InPanel, SSchematicGraphNode* InNode, const FDragDropOperation& InDragDropOperation)
+{
+	if (!ControlRigBlueprint.IsValid())
+	{
+		return;
+	}
+
+	if (!ControlRigBeingDebuggedPtr.IsValid())
+	{
+		return;
+	}
+
+	URigHierarchy* Hierarchy = ControlRigBeingDebuggedPtr->GetHierarchy();
+	if (!Hierarchy)
+	{
+		return;
+	}
+
+	FRigElementKey DraggedKey(InNode->NodeData->Name);
+	if (!DraggedKey.IsValid())
+	{
+		return;
+	}
+
+	FRigBaseElement* Element = Hierarchy->Find(DraggedKey);
+	if (!Element)
+	{
+		return;
+	}
+	
+	FRigConnectorElement* Connector = Cast<FRigConnectorElement>(Element);
+	if (!Connector)
+	{
+		return;
+	}
+
+	UModularRig* ModularRig = Cast<UModularRig>(ControlRigBeingDebuggedPtr);
+	if (!ModularRig)
+	{
+		return;
+	}
+
+	FString ModulePath = Hierarchy->GetNameMetadata(DraggedKey, URigHierarchy::NameSpaceMetadataName, NAME_None).ToString();
+	ModulePath.RemoveFromEnd(UModularRig::NamespaceSeparator);
+	const FRigModuleInstance* ModuleInstance = ModularRig->FindModule(ModulePath);
+	if (!ModuleInstance)
+	{
+		return;
+	}
+
+	UModularRigRuleManager* RuleManager = Hierarchy->GetRuleManager();
+	FModularRigResolveResult Result = RuleManager->FindMatches(Connector, ModuleInstance, ControlRigBeingDebuggedPtr->GetElementKeyRedirector());
+
+	TArray<FRigElementResolveResult> Matches = Result.GetMatches();
+	TArray<FString> NameMatches;
+	Algo::Transform(Matches, NameMatches, [](const FRigElementResolveResult& Match)
+	{
+		return Match.GetKey().ToString();
+	});
+
+	// Fade all the unmatched nodes
+	for (FSchematicGraphNode* Node : Nodes)
+	{
+		Node->bFade = !NameMatches.ContainsByPredicate([Node](const FString& Match)
+		{
+			return Node->Name == Match;
+		});
+	};
+
+	// Unfade the ones included in the match
+	for (const FString& Match : NameMatches)
+	{
+		bool bExists = Nodes.ContainsByPredicate([Match](const FSchematicGraphNode* Node)
+		{
+			return Match == Node->Name;
+		});
+
+		// Create a temporary node that will be active only while this drag operation exists
+		if (!bExists)
+		{
+			FSchematicGraphNode* NewNode = AddNode(Match);
+			TemporaryNodes.Add(Match);
+		}
+	}
+}
+
+void FControlRigSchematicModel::HandleSchematicEndDrag(SSchematicGraphPanel* InPanel, SSchematicGraphNode* InNode, const FDragDropOperation& InDragDropOperation)
+{
+	for (FString& TempNode : TemporaryNodes)
+	{
+		RemoveNode(TempNode);
+	}
+	TemporaryNodes.Reset();
+
+	for (FSchematicGraphNode* Node : Nodes)
+	{
+		Node->bFade = false;
+	}
+}
+
 void FControlRigSchematicModel::HandleUpdateSchematicNodes(SSchematicGraphPanel* InPanel, TSharedPtr<SSchematicGraphNode> InNode)
 {
 	if (!ControlRigBlueprint.IsValid())
@@ -235,6 +336,15 @@ void FControlRigSchematicModel::HandleUpdateSchematicNodes(SSchematicGraphPanel*
 					}
 				}
 			}
+			else
+			{
+				Transform = Hierarchy->GetGlobalTransform(ElementKey);
+				FVector2D PixelPos = ControlRigEditor.Pin()->ComputePersonaProjectedScreenPos(Transform.GetLocation());
+				InNode->SetPosition(PixelPos, true);
+
+				InNode->Size->Set(FVector2d(20, 20));
+				InNode->Brush.TintColor = FStyleColors::AccentBlue;
+			}
 		}
 
 	}
diff --git a/Engine/Plugins/Animation/ControlRig/Source/ControlRigEditor/Private/Editor/ControlRigSchematicModel.h b/Engine/Plugins/Animation/ControlRig/Source/ControlRigEditor/Private/Editor/ControlRigSchematicModel.h
--- a/Engine/Plugins/Animation/ControlRig/Source/ControlRigEditor/Private/Editor/ControlRigSchematicModel.h
+++ b/Engine/Plugins/Animation/ControlRig/Source/ControlRigEditor/Private/Editor/ControlRigSchematicModel.h
@@ -25,6 +25,8 @@ struct FControlRigSchematicModel : FSchematicGraph
 
 	void HandleUpdateSchematicNodes(SSchematicGraphPanel* InPanel, TSharedPtr<SSchematicGraphNode> InNode);
 	void HandleSchematicNodeClicked(SSchematicGraphPanel* InPanel, SSchematicGraphNode* InNode);
+	void HandleSchematicBeginDrag(SSchematicGraphPanel* InPanel, SSchematicGraphNode* InNode, const FDragDropOperation& InDragDropOperation);
+	void HandleSchematicEndDrag(SSchematicGraphPanel* InPanel, SSchematicGraphNode* InNode, const FDragDropOperation& InDragDropOperation);
 	void HandleSchematicDrop(SSchematicGraphPanel* InPanel, SSchematicGraphNode* InNode, const FDragDropEvent& InDragDropEvent);
 
 private:
@@ -32,4 +34,6 @@ private:
 	TWeakPtr<FControlRigEditor> ControlRigEditor;
 	TWeakObjectPtr<UControlRigBlueprint> ControlRigBlueprint;
 	TWeakObjectPtr<UControlRig> ControlRigBeingDebuggedPtr;
+
+	TArray<FString> TemporaryNodes;
 };
\ No newline at end of file
diff --git a/Engine/Source/Developer/AnimationWidgets/Private/SchematicGraphPanel/SSchematicGraphPanel.cpp b/Engine/Source/Developer/AnimationWidgets/Private/SchematicGraphPanel/SSchematicGraphPanel.cpp
--- a/Engine/Source/Developer/AnimationWidgets/Private/SchematicGraphPanel/SSchematicGraphPanel.cpp
+++ b/Engine/Source/Developer/AnimationWidgets/Private/SchematicGraphPanel/SSchematicGraphPanel.cpp
@@ -19,14 +19,14 @@ void FSchematicGraph::Reset()
 	}
 }
 
-bool FSchematicGraph::AddNode(const FString& InName)
+FSchematicGraphNode* FSchematicGraph::AddNode(const FString& InName)
 {
 	if (Nodes.ContainsByPredicate([InName](const FSchematicGraphNode* Node)
 	{
 		return Node->Name == InName;
 	}))
 	{
-		return false;
+		return nullptr;
 	}
 
 	FSchematicGraphNode* NewElement = static_cast<FSchematicGraphNode*>(FMemory::Malloc(sizeof(FSchematicGraphNode)));
@@ -38,7 +38,7 @@ bool FSchematicGraph::AddNode(const FString& InName)
 	{
 		OnNodeAddedDelegate.Broadcast(NewElement);
 	}
-	return true;
+	return NewElement;
 }
 
 bool FSchematicGraph::RenameNode(const FString& InOldName, const FString& InNewName)
@@ -63,27 +63,30 @@ bool FSchematicGraph::RenameNode(const FString& InOldName, const FString& InNewN
 
 bool FSchematicGraph::RemoveNode(const FString& InName)
 {
-	FSchematicGraphNode** FoundNode = Nodes.FindByPredicate([InName](const FSchematicGraphNode* Node)
+	FSchematicGraphNode** FoundNodePtr = Nodes.FindByPredicate([InName](const FSchematicGraphNode* Node)
 	{
 		return Node->Name == InName;
 	});
-	if (FoundNode)
+	if (FoundNodePtr)
 	{
+		FSchematicGraphNode* FoundNode = *FoundNodePtr;
 		if (OnNodeRemovedDelegate.IsBound())
 		{
-			OnNodeRemovedDelegate.Broadcast(*FoundNode);
+			OnNodeRemovedDelegate.Broadcast(FoundNode);
 		}
-		Nodes.Remove(*FoundNode);
-		FMemory::Free(*FoundNode);
+		Nodes.Remove(FoundNode);
+		FMemory::Free(FoundNode);
 		return true;
 	}
 	return false;
 }
 
-TSharedRef<FSchematicGraphNodeDragDropOp> FSchematicGraphNodeDragDropOp::New(const TArray<FString>& InElements)
+TSharedRef<FSchematicGraphNodeDragDropOp> FSchematicGraphNodeDragDropOp::New(TArray<SSchematicGraphNode*> InSchematicGraphNodes, const TArray<FString>& InElements, FSchematicGraphNodeDragDropOp::FOnEndDrag InOnEndDragDelegate)
 {
 	TSharedRef<FSchematicGraphNodeDragDropOp> Operation = MakeShared<FSchematicGraphNodeDragDropOp>();
+	Operation->SchematicGraphNodes = InSchematicGraphNodes;
 	Operation->Elements = InElements;
+	Operation->OnEndDragDelegate = InOnEndDragDelegate;
 	Operation->Construct();
 	return Operation;
 }
@@ -113,6 +116,8 @@ void SSchematicGraphNode::Construct(const FArguments& InArgs)
 {
 	NodeData = InArgs._NodeData;
 	OnClickedDelegate = InArgs._OnClicked;
+	OnBeginDragDelegate = InArgs._OnBeginDrag;
+	OnEndDragDelegate = InArgs._OnEndDrag;
 	OnDropDelegate = InArgs._OnDrop;
 	TEasingAttributeInterpolator<FVector2d>::FSettings Vector2DInterpSettings(EEasingInterpolatorType::BounceEaseInOut, 0.2f);
 	Position = TAnimatedAttribute<FVector2d>::Create(Vector2DInterpSettings, FVector2d::ZeroVector);
@@ -138,13 +143,19 @@ int32 SSchematicGraphNode::OnPaint(const FPaintArgs& Args, const FGeometry& Allo
 	FVector2d SizeOffset = (CurSize-OriginalSize)*-0.5;
 	FVector2d TotalOffset = SizeOffset;
 
+	FLinearColor TintColor = Brush.TintColor.GetSpecifiedColor();
+	if (NodeData->bFade)
+	{
+		TintColor.A = 0.3;
+	}
+
 	FSlateDrawElement::MakeBox(
 				OutDrawElements,
 				NewLayerId,
 				AllottedGeometry.ToPaintGeometry(CurSize, FSlateLayoutTransform(TotalOffset)),
 				&Brush,
 				ESlateDrawEffect::None,
-				Brush.TintColor.GetSpecifiedColor()
+				TintColor
 				);
 	
 	return NewLayerId;
@@ -177,8 +188,25 @@ void SSchematicGraphNode::OnDragLeave(const FDragDropEvent& DragDropEvent)
 
 FReply SSchematicGraphNode::OnDrop(const FGeometry& MyGeometry, const FDragDropEvent& DragDropEvent)
 {
+	// Avoid dropping onto itself
+	TSharedPtr<FSchematicGraphNodeDragDropOp> SchematicDragDropOp = DragDropEvent.GetOperationAs<FSchematicGraphNodeDragDropOp>();
+	if (SchematicDragDropOp)
+	{
+		if (!SchematicDragDropOp->GetElements().IsEmpty())
+		{
+			if (SchematicDragDropOp->GetElements().ContainsByPredicate([this](const FString& Name)
+			{
+				return Name == NodeData->Name;
+			}))
+			{
+				return FReply::Unhandled();
+			}
+		}
+	}
+	
 	SNode::OnDrop(MyGeometry, DragDropEvent);
 	OnDropDelegate.ExecuteIfBound(this, DragDropEvent);
+	OnEndDragDelegate.ExecuteIfBound(this, *DragDropEvent.GetOperation().Get());
 	return FReply::Unhandled();
 }
 
@@ -188,7 +216,8 @@ FReply SSchematicGraphNode::OnDragDetected(const FGeometry& MyGeometry, const FP
 	if (MouseEvent.IsMouseButtonDown(EKeys::LeftMouseButton) && DraggedElements.Num() > 0)
 	{
 		bIsBeingDragged = true;
-		TSharedRef<FSchematicGraphNodeDragDropOp> DragDropOp = FSchematicGraphNodeDragDropOp::New(MoveTemp(DraggedElements));
+		TSharedRef<FSchematicGraphNodeDragDropOp> DragDropOp = FSchematicGraphNodeDragDropOp::New({this}, MoveTemp(DraggedElements), OnEndDragDelegate);
+		OnBeginDragDelegate.ExecuteIfBound(this, DragDropOp.Get());
 		return FReply::Handled().BeginDragDrop(DragDropOp);
 	}
 	
@@ -248,6 +277,8 @@ void SSchematicGraphPanel::Construct(const FArguments& InArgs)
 	PaddingInterNode = InArgs._PaddingInterNode;
 	UpdateNodeWidgetDelegate = InArgs._OnUpdateNodeWidget;
 	OnNodeClickedDelegate = InArgs._OnNodeClicked;
+	OnBeginDragDelegate = InArgs._OnBeginDrag;
+	OnEndDragDelegate = InArgs._OnEndDrag;
 	OnDropDelegate = InArgs._OnDrop;
 
 	TEasingAttributeInterpolator<float>::FSettings FloatInterpSettings(EEasingInterpolatorType::BounceEaseInOut, 0.2f);
@@ -282,6 +313,8 @@ void SSchematicGraphPanel::AddNode(FSchematicGraphNode* NodeToAdd)
 {
 	TSharedRef<SSchematicGraphNode> NewNode = SNew(SSchematicGraphNode)
 														.OnClicked_Raw(this, &SSchematicGraphPanel::OnNodeClicked)
+														.OnBeginDrag_Raw(this, &SSchematicGraphPanel::OnBeginDragEvent)
+														.OnEndDrag_Raw(this, &SSchematicGraphPanel::OnEndDragEvent)
 														.OnDrop_Raw(this, &SSchematicGraphPanel::OnDropEvent)
 														.NodeData(NodeToAdd);
 	SNodePanel::AddGraphNode(NewNode);
@@ -428,10 +461,6 @@ void SSchematicGraphPanel::Tick(float DeltaTime)
 			bIsDragging = true;
 		}
 	}
-	else
-	{
-		bIsDragging = false;
-	}
 	SetFadeBackground(bIsDragging);
 
 	uint16 TopRightNodes = 0;
@@ -490,7 +519,7 @@ void SSchematicGraphPanel::Tick(float DeltaTime)
 				bImmediatePosition = true;
 				
 			}
-			else
+			else if (DeltaTime > 0.f)
 			{
 				SNode->bIsBeingDragged = false;
 			}
@@ -511,6 +540,16 @@ void SSchematicGraphPanel::OnNodeClicked(SSchematicGraphNode* Node)
 	OnNodeClickedDelegate.ExecuteIfBound(this, Node);
 }
 
+void SSchematicGraphPanel::OnBeginDragEvent(SSchematicGraphNode* Node, const FDragDropOperation& InDragDropOp)
+{
+	OnBeginDragDelegate.ExecuteIfBound(this, Node, InDragDropOp);
+}
+
+void SSchematicGraphPanel::OnEndDragEvent(SSchematicGraphNode* Node, const FDragDropOperation& InDragDropOp)
+{
+	OnEndDragDelegate.ExecuteIfBound(this, Node, InDragDropOp);
+}
+
 void SSchematicGraphPanel::OnDropEvent(SSchematicGraphNode* Node, const FDragDropEvent& InDragDropEvent)
 {
 	OnDropDelegate.ExecuteIfBound(this, Node, InDragDropEvent);
diff --git a/Engine/Source/Developer/AnimationWidgets/Public/SchematicGraphPanel/SSchematicGraphPanel.h b/Engine/Source/Developer/AnimationWidgets/Public/SchematicGraphPanel/SSchematicGraphPanel.h
--- a/Engine/Source/Developer/AnimationWidgets/Public/SchematicGraphPanel/SSchematicGraphPanel.h
+++ b/Engine/Source/Developer/AnimationWidgets/Public/SchematicGraphPanel/SSchematicGraphPanel.h
@@ -20,6 +20,7 @@ struct ANIMATIONWIDGETS_API FSchematicGraphNode
 {
 	FString Name;
 	bool bIsSelected = false;
+	bool bFade = false;
 	ESchematicGraphNodePlacement Placement;
 };
 
@@ -35,7 +36,7 @@ public:
 	TMap<FSchematicGraphNode*, TArray<FSchematicGraphNode*>> Links;
 
 	void Reset();
-	bool AddNode(const FString& InName);
+	FSchematicGraphNode* AddNode(const FString& InName);
 	bool RenameNode(const FString& InOldName, const FString& InNewName);
 	bool RemoveNode(const FString& InName);
 
@@ -49,13 +50,28 @@ public:
 	FOnGraphReset& OnGraphReset() { return OnGraphResetDelegate; }
 };
 
+class SSchematicGraphNode;
+
 class FSchematicGraphNodeDragDropOp : public FDragDropOperation
 {
 public:
 	DRAG_DROP_OPERATOR_TYPE(FSchematicGraphNodeDragDropOp, FDragDropOperation)
 
-	static TSharedRef<FSchematicGraphNodeDragDropOp> New(const TArray<FString>& InElements);
+	DECLARE_DELEGATE_TwoParams(FOnEndDrag, SSchematicGraphNode*, const FDragDropOperation&);
+
+	static TSharedRef<FSchematicGraphNodeDragDropOp> New(TArray<SSchematicGraphNode*> InSchematicGraphNodes, const TArray<FString>& InElements, FOnEndDrag InOnEndDragDelegate);
 
+	~FSchematicGraphNodeDragDropOp()
+	{
+		if (OnEndDragDelegate.IsBound())
+		{
+			for (SSchematicGraphNode* Node : SchematicGraphNodes)
+			{
+				OnEndDragDelegate.Execute(Node, *this);
+			}
+		}
+	}
+	
 	virtual TSharedPtr<SWidget> GetDefaultDecorator() const override;
 
 	/** @return true if this drag operation contains property paths */
@@ -74,19 +90,29 @@ public:
 
 private:
 
+	/** Nodes being dragged */
+	TArray<SSchematicGraphNode*> SchematicGraphNodes;
+
 	/** Data for the property paths this item represents */
 	TArray<FString> Elements;
+
+	/** Delegate to call when this drag operation ends */
+	FOnEndDrag OnEndDragDelegate;
 };
 
 class ANIMATIONWIDGETS_API SSchematicGraphNode : public SNodePanel::SNode
 {
 public:
 	DECLARE_DELEGATE_OneParam(FOnClicked, SSchematicGraphNode*);
+	DECLARE_DELEGATE_TwoParams(FOnBeginDrag, SSchematicGraphNode*, const FDragDropOperation&);
+	DECLARE_DELEGATE_TwoParams(FOnEndDrag, SSchematicGraphNode*, const FDragDropOperation&);
 	DECLARE_DELEGATE_TwoParams(FOnDrop, SSchematicGraphNode*, const FDragDropEvent&);
 
 	SLATE_BEGIN_ARGS(SSchematicGraphNode){}
 	SLATE_ARGUMENT(FSchematicGraphNode*, NodeData)
 	SLATE_EVENT(FOnClicked, OnClicked)
+	SLATE_EVENT(FOnBeginDrag, OnBeginDrag)
+	SLATE_EVENT(FOnEndDrag, OnEndDrag)
 	SLATE_EVENT(FOnDrop, OnDrop)
 	SLATE_END_ARGS()
 
@@ -113,6 +139,8 @@ public:
 	float ScaledDown = 0.75;
 
 	FOnClicked OnClickedDelegate;
+	FOnBeginDrag OnBeginDragDelegate;
+	FOnEndDrag OnEndDragDelegate;
 	FOnDrop OnDropDelegate;
 	FSchematicGraphNode* NodeData;
 	TSharedPtr<TAnimatedAttribute<FVector2d>> Position;
@@ -129,6 +157,8 @@ public:
 	
 	DECLARE_DELEGATE_TwoParams(FUpdateNodeWidget, SSchematicGraphPanel*, TSharedPtr<SSchematicGraphNode>);
 	DECLARE_DELEGATE_TwoParams(FOnNodeClicked, SSchematicGraphPanel*, SSchematicGraphNode*);
+	DECLARE_DELEGATE_ThreeParams(FOnBeginDrag, SSchematicGraphPanel*, SSchematicGraphNode*, const FDragDropOperation&);
+	DECLARE_DELEGATE_ThreeParams(FOnEndDrag, SSchematicGraphPanel*, SSchematicGraphNode*, const FDragDropOperation&);
 	DECLARE_DELEGATE_ThreeParams(FOnDrop, SSchematicGraphPanel*, SSchematicGraphNode*, const FDragDropEvent&);
 	
 	SLATE_BEGIN_ARGS(SSchematicGraphPanel) {}
@@ -141,6 +171,8 @@ public:
 	SLATE_ARGUMENT(int32, PaddingInterNode)
 	SLATE_EVENT(FUpdateNodeWidget, OnUpdateNodeWidget)
 	SLATE_EVENT(FOnNodeClicked, OnNodeClicked)
+	SLATE_EVENT(FOnBeginDrag, OnBeginDrag)
+	SLATE_EVENT(FOnEndDrag, OnEndDrag)
 	SLATE_EVENT(FOnDrop, OnDrop)
 	SLATE_END_ARGS()
 
@@ -173,6 +205,8 @@ public:
 	// End of FTickableEditorObject interface
 
 	void OnNodeClicked(SSchematicGraphNode* Node);
+	void OnBeginDragEvent(SSchematicGraphNode* Node, const FDragDropOperation& InDragDropEvent);
+	void OnEndDragEvent(SSchematicGraphNode* Node, const FDragDropOperation& InDragDropEvent);
 	void OnDropEvent(SSchematicGraphNode* Node, const FDragDropEvent& InDragDropEvent);
 
 	void SetFadeBackground(bool bInFade) { FadeBackgroundAlpha->Set(bInFade ? 0.5f : 0.f); }
@@ -188,6 +222,8 @@ public:
 	FSchematicGraph* GraphData;
 	FUpdateNodeWidget UpdateNodeWidgetDelegate;
 	FOnNodeClicked OnNodeClickedDelegate;
+	FOnBeginDrag OnBeginDragDelegate;
+	FOnEndDrag OnEndDragDelegate;
 	FOnDrop OnDropDelegate;
 };
 
