Control Rig: Refactor schematic view

#rnx

diff --git a/Engine/Plugins/Animation/ControlRig/Source/ControlRig/Private/ModularRigController.cpp b/Engine/Plugins/Animation/ControlRig/Source/ControlRig/Private/ModularRigController.cpp
--- a/Engine/Plugins/Animation/ControlRig/Source/ControlRig/Private/ModularRigController.cpp
+++ b/Engine/Plugins/Animation/ControlRig/Source/ControlRig/Private/ModularRigController.cpp
@@ -214,6 +214,70 @@ bool UModularRigController::ConnectConnectorToElement(const FRigElementKey& InCo
 	return true;
 }
 
+bool UModularRigController::DisconnectConnector(const FRigElementKey& InConnectorKey, bool bSetupUndo)
+{
+	FString ConnectorParentPath, ConnectorName;
+	if (!InConnectorKey.Name.ToString().Split(UModularRig::NamespaceSeparator, &ConnectorParentPath, &ConnectorName, ESearchCase::CaseSensitive, ESearchDir::FromEnd))
+	{
+		UE_LOG(LogControlRig, Error, TEXT("Connector %s does not contain a namespace"), *InConnectorKey.ToString());
+		return false;
+	}
+	
+	FRigModuleReference* Module = FindModule(ConnectorParentPath);
+	if (!Module)
+	{
+		UE_LOG(LogControlRig, Error, TEXT("Could not find module %s"), *ConnectorParentPath);
+		return false;
+	}
+
+	UControlRig* RigCDO = Module->Class->GetDefaultObject<UControlRig>();
+	if (!RigCDO)
+	{
+		UE_LOG(LogControlRig, Error, TEXT("Invalid rig module class %s"), *Module->Class->GetPathName());
+		return false;
+	}
+
+	const FRigModuleConnector* ModuleConnector = RigCDO->GetRigModuleSettings().ExposedConnectors.FindByPredicate(
+		[ConnectorName](FRigModuleConnector& Connector)
+		{
+			return Connector.Name == ConnectorName;
+		});
+	if (!ModuleConnector)
+	{
+		UE_LOG(LogControlRig, Error, TEXT("Could not find connector %s in class %s"), *ConnectorName, *Module->Class->GetPathName());
+		return false;
+	}
+
+	UBlueprint* Blueprint = Cast<UBlueprint>(GetOuter());
+	const IRigHierarchyProvider* HierarchyProvider = CastChecked<IRigHierarchyProvider>(Blueprint);
+	const FRigConnectorElement* Connector = Cast<FRigConnectorElement>(HierarchyProvider->GetHierarchy()->Find(InConnectorKey));
+	if (!Connector)
+	{
+		UE_LOG(LogControlRig, Error, TEXT("Could not find connector %s"), *InConnectorKey.ToString());
+		return false;
+	}
+
+#if WITH_EDITOR
+	TSharedPtr<FScopedTransaction> TransactionPtr;
+	if (bSetupUndo)
+	{
+		TransactionPtr = MakeShared<FScopedTransaction>(NSLOCTEXT("ModularRigController", "ConnectModuleToElementTransaction", "Connect to Element"));
+		Blueprint->Modify();
+	}
+#endif 
+
+	const FRigElementKey ConnectorKey(*ConnectorName, ERigElementType::Connector);
+	Module->Connections.Remove(ConnectorKey);
+
+	Notify(EModularRigNotification::ConnectionChanged, Module);
+
+#if WITH_EDITOR
+	TransactionPtr.Reset();
+#endif
+	
+	return true;
+}
+
 bool UModularRigController::SetConfigValueInModule(const FString& InModulePath, const FName& InVariableName, const FString& InValue, bool bSetupUndo)
 {
 	FRigModuleReference* Module = FindModule(InModulePath);
@@ -537,7 +601,8 @@ bool UModularRigController::DeleteModule(const FString& InModulePath, bool bSetu
 #endif
 
 	// Unparent children (add them to root)
-	for (FRigModuleReference* Child : Module->CachedChildren)
+	TArray<FRigModuleReference*> PreviousChildren = Module->CachedChildren;
+	for (const FRigModuleReference* Child : PreviousChildren)
 	{
 		ReparentModule(Child->GetPath(), FString(), bSetupUndo);
 	}
diff --git a/Engine/Plugins/Animation/ControlRig/Source/ControlRig/Private/Rigs/RigConnectionRules.cpp b/Engine/Plugins/Animation/ControlRig/Source/ControlRig/Private/Rigs/RigConnectionRules.cpp
--- a/Engine/Plugins/Animation/ControlRig/Source/ControlRig/Private/Rigs/RigConnectionRules.cpp
+++ b/Engine/Plugins/Animation/ControlRig/Source/ControlRig/Private/Rigs/RigConnectionRules.cpp
@@ -107,7 +107,7 @@ const FRigConnectorElement* FRigConnectionRuleInput::FindPrimaryConnector(FText*
 	check(Hierarchy);
 	check(Module);
 
-	const FName ModuleNameSpace = *Module->GetPath();
+	const FName ModuleNameSpace = *Module->GetNamespace();
 
 	const FRigConnectorElement* PrimaryConnector = nullptr;
 	Hierarchy->ForEach<FRigConnectorElement>(
@@ -143,7 +143,7 @@ TArray<const FRigConnectorElement*> FRigConnectionRuleInput::FindSecondaryConnec
 	check(Hierarchy);
 	check(Module);
 
-	const FName ModuleNameSpace = *Module->GetPath();
+	const FName ModuleNameSpace = *Module->GetNamespace();
 
 	TArray<const FRigConnectorElement*> SecondaryConnectors;
 	Hierarchy->ForEach<FRigConnectorElement>(
diff --git a/Engine/Plugins/Animation/ControlRig/Source/ControlRig/Private/Rigs/RigHierarchy.cpp b/Engine/Plugins/Animation/ControlRig/Source/ControlRig/Private/Rigs/RigHierarchy.cpp
--- a/Engine/Plugins/Animation/ControlRig/Source/ControlRig/Private/Rigs/RigHierarchy.cpp
+++ b/Engine/Plugins/Animation/ControlRig/Source/ControlRig/Private/Rigs/RigHierarchy.cpp
@@ -1233,6 +1233,10 @@ FRigName URigHierarchy::JoinNameSpace(const FRigName& InLeft, const FRigName& In
 
 FString URigHierarchy::JoinNameSpace(const FString& InLeft, const FString& InRight)
 {
+	if(InLeft.EndsWith(UModularRig::NamespaceSeparator))
+	{
+		return InLeft + InRight;
+	}
 	return InLeft + UModularRig::NamespaceSeparator + InRight;
 }
 
diff --git a/Engine/Plugins/Animation/ControlRig/Source/ControlRig/Private/Units/Execution/RigUnit_Hierarchy.cpp b/Engine/Plugins/Animation/ControlRig/Source/ControlRig/Private/Units/Execution/RigUnit_Hierarchy.cpp
--- a/Engine/Plugins/Animation/ControlRig/Source/ControlRig/Private/Units/Execution/RigUnit_Hierarchy.cpp
+++ b/Engine/Plugins/Animation/ControlRig/Source/ControlRig/Private/Units/Execution/RigUnit_Hierarchy.cpp
@@ -182,7 +182,16 @@ FRigUnit_HierarchyGetChainItemArray_Execute()
 			TArray<FRigElementKey> Keys;
 
 			const FRigTransformElement* StartElement = Cast<FRigTransformElement>(CachedStart.GetElement());
+			if(StartElement && (StartElement->GetType() == ERigElementType::Socket))
+			{
+				StartElement = Cast<FRigTransformElement>(ExecuteContext.Hierarchy->GetFirstParent(StartElement));
+			}
 			const FRigTransformElement* EndElement = Cast<FRigTransformElement>(CachedEnd.GetElement());
+			if(EndElement && (EndElement->GetType() == ERigElementType::Socket))
+			{
+				EndElement = Cast<FRigTransformElement>(ExecuteContext.Hierarchy->GetFirstParent(EndElement));
+			}
+			
 			if(StartElement == nullptr || EndElement == nullptr)
 			{
 				Keys.Reset();
diff --git a/Engine/Plugins/Animation/ControlRig/Source/ControlRig/Public/ModularRigController.h b/Engine/Plugins/Animation/ControlRig/Source/ControlRig/Public/ModularRigController.h
--- a/Engine/Plugins/Animation/ControlRig/Source/ControlRig/Public/ModularRigController.h
+++ b/Engine/Plugins/Animation/ControlRig/Source/ControlRig/Public/ModularRigController.h
@@ -30,6 +30,9 @@ class CONTROLRIG_API UModularRigController : public UObject
 	UFUNCTION(BlueprintCallable, Category = "Control Rig | Modules")
 	bool ConnectConnectorToElement(const FRigElementKey& InConnectorKey, const FRigElementKey& InTargetKey, bool bSetupUndo = true);
 
+	UFUNCTION(BlueprintCallable, Category = "Control Rig | Modules")
+	bool DisconnectConnector(const FRigElementKey& InConnectorKey, bool bSetupUndo = true);
+
 	UFUNCTION(BlueprintCallable, Category = "Control Rig | Modules")
 	bool SetConfigValueInModule(const FString& InModulePath, const FName& InVariableName, const FString& InValue, bool bSetupUndo = true);
 
diff --git a/Engine/Plugins/Animation/ControlRig/Source/ControlRig/Public/Rigs/RigConnectionRules.h b/Engine/Plugins/Animation/ControlRig/Source/ControlRig/Public/Rigs/RigConnectionRules.h
--- a/Engine/Plugins/Animation/ControlRig/Source/ControlRig/Public/Rigs/RigConnectionRules.h
+++ b/Engine/Plugins/Animation/ControlRig/Source/ControlRig/Public/Rigs/RigConnectionRules.h
@@ -163,7 +163,7 @@ struct CONTROLRIG_API FRigTypeConnectionRule : public FRigConnectionRule
 public:
 
 	FRigTypeConnectionRule()
-		: ElementType(ERigElementType::Connector)
+		: ElementType(ERigElementType::Socket)
 	{}
 
 	FRigTypeConnectionRule(ERigElementType InElementType)
diff --git a/Engine/Plugins/Animation/ControlRig/Source/ControlRig/Public/Rigs/RigHierarchyCache.h b/Engine/Plugins/Animation/ControlRig/Source/ControlRig/Public/Rigs/RigHierarchyCache.h
--- a/Engine/Plugins/Animation/ControlRig/Source/ControlRig/Public/Rigs/RigHierarchyCache.h
+++ b/Engine/Plugins/Animation/ControlRig/Source/ControlRig/Public/Rigs/RigHierarchyCache.h
@@ -188,5 +188,5 @@ private:
 	friend class URigHierarchyController;
 	friend class FRigModuleInstanceDetails;
 	friend class UModularRig;
-	friend struct FControlRigSchematicModel;
+	friend class FControlRigSchematicModel;
 };
\ No newline at end of file
diff --git a/Engine/Plugins/Animation/ControlRig/Source/ControlRigDeveloper/Private/ControlRigBlueprint.cpp b/Engine/Plugins/Animation/ControlRig/Source/ControlRigDeveloper/Private/ControlRigBlueprint.cpp
--- a/Engine/Plugins/Animation/ControlRig/Source/ControlRigDeveloper/Private/ControlRigBlueprint.cpp
+++ b/Engine/Plugins/Animation/ControlRig/Source/ControlRigDeveloper/Private/ControlRigBlueprint.cpp
@@ -1179,7 +1179,6 @@ void UControlRigBlueprint::HandlePackageDone()
 			return true;
 		});
 
-		RefreshModuleConnectors();
 		RecompileModularRig();
 	}
 }
@@ -1587,6 +1586,7 @@ UModularRigController* UControlRigBlueprint::GetModularRigController()
 
 void UControlRigBlueprint::RecompileModularRig()
 {
+	RefreshModuleConnectors();
 	OnModularRigPreCompiled().Broadcast(this);
 	if (const UClass* MyControlRigClass = GeneratedClass)
 	{
@@ -2325,10 +2325,7 @@ void UControlRigBlueprint::RefreshModuleConnectors(const FRigModuleReference* In
 			// setup the module information. this is needed so that newly added
 			// connectors result in the right namespace metadata etc
 			FRigVMExtendedExecuteContext& Context = CDO->GetRigVMExtendedExecuteContext();
-			FRigHierarchyExecuteContextBracket HierarchyContextGuard(Controller->GetHierarchy(), &Context);
 			FControlRigExecuteContext& PublicContext = Context.GetPublicDataSafe<FControlRigExecuteContext>();
-			FControlRigExecuteContextRigModuleGuard RigModuleGuard(PublicContext, InModule->GetNamespace());
-		
 			const UControlRig* ModuleCDO = InModule->Class->GetDefaultObject<UControlRig>();
 			const TArray<FRigModuleConnector>& ExpectedConnectors = ModuleCDO->GetRigModuleSettings().ExposedConnectors;
 
@@ -2348,6 +2345,7 @@ void UControlRigBlueprint::RefreshModuleConnectors(const FRigModuleReference* In
 				if(bRemoveAllConnectors || !bConnectorExpected)
 				{
 					(void)Controller->RemoveElement(Connector);
+					ConnectionMap.Remove(Connector);
 				}
 			}
 
@@ -2357,10 +2355,20 @@ void UControlRigBlueprint::RefreshModuleConnectors(const FRigModuleReference* In
 				for (const FRigModuleConnector& Connector : ExpectedConnectors)
 				{
 					const FName ConnectorName = *Connector.Name;
-					if(!Hierarchy->Contains(FRigElementKey(ConnectorName, ERigElementType::Connector)))
+					const FName ConnectorNameWithNameSpace = *URigHierarchy::JoinNameSpace(InModule->GetNamespace(), Connector.Name);
+					const FRigElementKey ConnectorKeyWithNameSpace(ConnectorNameWithNameSpace, ERigElementType::Connector);
+					if(!Hierarchy->Contains(ConnectorKeyWithNameSpace))
 					{
+						FRigHierarchyExecuteContextBracket HierarchyContextGuard(Hierarchy, &Context);
+						FControlRigExecuteContextRigModuleGuard RigModuleGuard(PublicContext, InModule->GetNamespace());
 						(void)Controller->AddConnector(ConnectorName, Connector.Settings);
 					}
+					else
+					{
+						// copy the connector settings
+						FRigConnectorElement* ExistingConnector = Hierarchy->FindChecked<FRigConnectorElement>(ConnectorKeyWithNameSpace);
+						ExistingConnector->Settings = Connector.Settings;
+					}
 				}
 			}
 			
@@ -2568,11 +2576,30 @@ void UControlRigBlueprint::HandleRigModulesModified(EModularRigNotification InNo
 			Hierarchy->Modify();
 			
 			const FString Namespace = InModule->GetNamespace();
+
+			TArray<FRigElementKey> ConnectionsToRemove;
+			for (TPair<FRigElementKey, FRigElementKey> ExistingConnection : ConnectionMap)
+			{
+				if(ExistingConnection.Key.Name.ToString().StartsWith(Namespace, ESearchCase::CaseSensitive))
+				{
+					if(!InModule->Connections.Contains(ExistingConnection.Key))
+					{
+						ConnectionsToRemove.Add(ExistingConnection.Key);
+					}
+				}
+			}
+				
+			for(const FRigElementKey& ConnectionToRemove : ConnectionsToRemove)
+			{
+				ConnectionMap.Remove(ConnectionToRemove);
+			}
+				
 			for (TPair<FRigElementKey, FRigElementKey> Connection : InModule->Connections)
 			{
 				const FString NamespacedName = FString::Printf(TEXT("%s%s"), *Namespace, *Connection.Key.Name.ToString());
 				ConnectionMap.FindOrAdd(FRigElementKey(*NamespacedName, ERigElementType::Connector)) = Connection.Value;
 			}
+				
 			HierarchyModifiedEvent.Broadcast(ERigHierarchyNotification::HierarchyReset, Hierarchy, nullptr);
 			break;
 		}
diff --git a/Engine/Plugins/Animation/ControlRig/Source/ControlRigEditor/Private/ControlRigEditorStyle.h b/Engine/Plugins/Animation/ControlRig/Source/ControlRigEditor/Private/ControlRigEditorStyle.h
--- a/Engine/Plugins/Animation/ControlRig/Source/ControlRigEditor/Private/ControlRigEditorStyle.h
+++ b/Engine/Plugins/Animation/ControlRig/Source/ControlRigEditor/Private/ControlRigEditorStyle.h
@@ -151,6 +151,9 @@ public:
 			Set("ControlRig.Schematic.ConnectorPrimary", new IMAGE_BRUSH_SVG("Slate/Connector_Primary", Icon128x128));
 			Set("ControlRig.Schematic.ConnectorSecondary", new IMAGE_BRUSH_SVG("Slate/Connector_Secondary", Icon128x128));
 			Set("ControlRig.Schematic.ConnectorOptional", new IMAGE_BRUSH_SVG("Slate/Connector_Optional", Icon128x128));
+			Set("ControlRig.Schematic.Bone", new IMAGE_BRUSH_SVG("Slate/Bone_Schematic", Icon128x128));
+			Set("ControlRig.Schematic.Control", new IMAGE_BRUSH_SVG("Slate/Control_Schematic", Icon128x128));
+			Set("ControlRig.Schematic.Null", new IMAGE_BRUSH_SVG("Slate/Null_Schematic", Icon128x128));
 		}
 
 		FSlateStyleRegistry::RegisterSlateStyle(*this);
diff --git a/Engine/Plugins/Animation/ControlRig/Source/ControlRigEditor/Private/ControlRigModuleDetails.cpp b/Engine/Plugins/Animation/ControlRig/Source/ControlRigEditor/Private/ControlRigModuleDetails.cpp
--- a/Engine/Plugins/Animation/ControlRig/Source/ControlRigEditor/Private/ControlRigModuleDetails.cpp
+++ b/Engine/Plugins/Animation/ControlRig/Source/ControlRigEditor/Private/ControlRigModuleDetails.cpp
@@ -705,7 +705,14 @@ void FRigModuleInstanceDetails::OnElementNameChanged(TSharedPtr<FString> InItem,
 				if(const FRigModuleInstance* Module = Info.GetModule())
 				{
 					FRigElementKey NamespacedConnector(*URigHierarchy::JoinNameSpace(Module->GetPath(), Connector.Name.ToString()), ERigElementType::Connector);
-					Controller->ConnectConnectorToElement(NamespacedConnector, *TargetKey);
+					if(!Blueprint->GetDebuggedControlRig()->GetHierarchy()->Contains(*TargetKey))
+					{
+						Controller->DisconnectConnector(NamespacedConnector);
+					}
+					else
+					{
+						Controller->ConnectConnectorToElement(NamespacedConnector, *TargetKey);
+					}
 				}
 			}
 		}
@@ -728,7 +735,14 @@ void FRigModuleInstanceDetails::OnElementTypeChanged(ERigElementType InElementTy
 				if(const FRigModuleInstance* Module = Info.GetModule())
 				{
 					FRigElementKey NamespacedConnector(*URigHierarchy::JoinNameSpace(Module->GetPath(), Connector.Name.ToString()), ERigElementType::Connector);
-					Controller->ConnectConnectorToElement(NamespacedConnector, *TargetKey);
+					if(!Blueprint->GetDebuggedControlRig()->GetHierarchy()->Contains(*TargetKey))
+					{
+						Controller->DisconnectConnector(NamespacedConnector);
+					}
+					else
+					{
+						Controller->ConnectConnectorToElement(NamespacedConnector, *TargetKey);
+					}
 				}
 			}
 		}
diff --git a/Engine/Plugins/Animation/ControlRig/Source/ControlRigEditor/Private/Editor/ControlRigEditor.cpp b/Engine/Plugins/Animation/ControlRig/Source/ControlRigEditor/Private/Editor/ControlRigEditor.cpp
--- a/Engine/Plugins/Animation/ControlRig/Source/ControlRigEditor/Private/Editor/ControlRigEditor.cpp
+++ b/Engine/Plugins/Animation/ControlRig/Source/ControlRigEditor/Private/Editor/ControlRigEditor.cpp
@@ -2571,7 +2571,6 @@ void FControlRigEditor::CacheNameLists()
 
 void FControlRigEditor::HandleSchematicViewportCreated(const TSharedRef<SSchematicGraphPanel>& InViewport)
 {
-	InViewport->UpdateNodeWidgetDelegate.BindRaw(&SchematicModel, &FControlRigSchematicModel::HandleUpdateSchematicNodes);
 	InViewport->OnNodeClickedDelegate.BindRaw(&SchematicModel, &FControlRigSchematicModel::HandleSchematicNodeClicked);
 	InViewport->OnBeginDragDelegate.BindRaw(&SchematicModel, &FControlRigSchematicModel::HandleSchematicBeginDrag);
 	InViewport->OnEndDragDelegate.BindRaw(&SchematicModel, &FControlRigSchematicModel::HandleSchematicEndDrag);
diff --git a/Engine/Plugins/Animation/ControlRig/Source/ControlRigEditor/Private/Editor/ControlRigSchematicModel.cpp b/Engine/Plugins/Animation/ControlRig/Source/ControlRigEditor/Private/Editor/ControlRigSchematicModel.cpp
--- a/Engine/Plugins/Animation/ControlRig/Source/ControlRigEditor/Private/Editor/ControlRigSchematicModel.cpp
+++ b/Engine/Plugins/Animation/ControlRig/Source/ControlRigEditor/Private/Editor/ControlRigSchematicModel.cpp
@@ -10,6 +10,11 @@
 #include "DragAndDrop/AssetDragDropOp.h"
 #include "Rigs/RigHierarchyController.h"
 
+FString FControlRigSchematicRigElementKeyNode::GetDragDropDecoratorLabel() const
+{
+	return Key.ToString();
+}
+
 FControlRigSchematicModel::~FControlRigSchematicModel()
 {
 	if(ControlRigBeingDebuggedPtr.IsValid())
@@ -35,6 +40,67 @@ void FControlRigSchematicModel::SetEditor(const TSharedRef<FControlRigEditor>& I
 	ControlRigBeingDebuggedPtr = ControlRigBlueprint->GetDebuggedControlRig();
 }
 
+void FControlRigSchematicModel::Reset()
+{
+	FSchematicGraphModel::Reset();
+	RigElementKeyToGuid.Reset();
+}
+
+FControlRigSchematicRigElementKeyNode* FControlRigSchematicModel::AddElementKeyNode(const FRigElementKey& InKey, bool bNotify)
+{
+	FControlRigSchematicRigElementKeyNode* Node = AddNode<FControlRigSchematicRigElementKeyNode>(false);
+	if(Node)
+	{
+		Node->Key = InKey;
+		RigElementKeyToGuid.Add(Node->GetKey(), Node->GetGuid());
+
+		if (bNotify && OnNodeAddedDelegate.IsBound())
+		{
+			OnNodeAddedDelegate.Broadcast(Node);
+		}
+	}
+	return Node;
+}
+
+const FControlRigSchematicRigElementKeyNode* FControlRigSchematicModel::FindElementKeyNode(const FRigElementKey& InKey) const
+{
+	if(const FGuid* FoundGuid = RigElementKeyToGuid.Find(InKey))
+	{
+		return FindNode< FControlRigSchematicRigElementKeyNode >(*FoundGuid);
+	}
+	return nullptr;
+}
+
+bool FControlRigSchematicModel::ContainsElementKeyNode(const FRigElementKey& InKey) const
+{
+	return FindElementKeyNode(InKey) != nullptr;
+}
+
+bool FControlRigSchematicModel::RemoveNode(const FGuid& InGuid)
+{
+	FRigElementKey KeyToRemove;
+	if(const FControlRigSchematicRigElementKeyNode* Node = FindNode<FControlRigSchematicRigElementKeyNode>(InGuid))
+	{
+		KeyToRemove = Node->GetKey();
+	}
+	if(FSchematicGraphModel::RemoveNode(InGuid))
+	{
+		RigElementKeyToGuid.Remove(KeyToRemove);
+		return true;
+	}
+	return false;
+}
+
+bool FControlRigSchematicModel::RemoveElementKeyNode(const FRigElementKey& InKey)
+{
+	if(const FGuid* GuidPtr = RigElementKeyToGuid.Find(InKey))
+	{
+		const FGuid Guid = *GuidPtr;
+		return RemoveNode(Guid);
+	}
+	return false;
+}
+
 void FControlRigSchematicModel::OnSetObjectBeingDebugged(UObject* InObject)
 {
 	if(ControlRigBeingDebuggedPtr.Get() == InObject)
@@ -77,7 +143,11 @@ void FControlRigSchematicModel::OnHierarchyModified(ERigHierarchyNotification In
 				if (InElement->GetType() == ERigElementType::Socket ||
 					InElement->GetType() == ERigElementType::Connector)
 				{
-					AddNode(InElement->GetKey().ToString());
+					const FRigElementKey& NodeKey = InElement->GetKey();
+					if(!ContainsElementKeyNode(NodeKey))
+					{
+						AddElementKeyNode(NodeKey);
+					}
 				}
 			}
 			break;
@@ -91,7 +161,8 @@ void FControlRigSchematicModel::OnHierarchyModified(ERigHierarchyNotification In
 				{
 					const FString OldNameStr = InHierarchy->GetPreviousName(InElement->GetKey()).ToString();
 					FRigElementKey OldKey(*OldNameStr, InElement->GetType());
-					RenameNode(OldKey.ToString(), InElement->GetKey().ToString());
+					// todooo
+					//RenameNode(OldKey.ToString(), InElement->GetKey().ToString());
 				}
 			}
 			break;
@@ -103,7 +174,7 @@ void FControlRigSchematicModel::OnHierarchyModified(ERigHierarchyNotification In
 				if (InElement->GetType() == ERigElementType::Socket ||
 					InElement->GetType() == ERigElementType::Connector)
 				{
-					RemoveNode(InElement->GetKey().ToString());
+					RemoveElementKeyNode(InElement->GetKey());
 				}
 			}
 			break;
@@ -112,24 +183,36 @@ void FControlRigSchematicModel::OnHierarchyModified(ERigHierarchyNotification In
 		{
 			Reset();
 			TArray<FRigSocketElement*> Sockets = InHierarchy->GetElementsOfType<FRigSocketElement>();
-			for (FRigSocketElement* Socket : Sockets)
+			for (const FRigSocketElement* Socket : Sockets)
 			{
-				AddNode(Socket->GetKey().ToString());
+				const FString NodeKey = Socket->GetKey().ToString();
+				if(!ContainsElementKeyNode(NodeKey))
+				{
+					AddElementKeyNode(NodeKey);
+				}
 			}
 			if (ControlRigBeingDebuggedPtr.IsValid())
 			{
-				FRigElementKeyRedirector& Redirector = ControlRigBeingDebuggedPtr->GetElementKeyRedirector();
+				const FRigElementKeyRedirector& Redirector = ControlRigBeingDebuggedPtr->GetElementKeyRedirector();
 				TArray<FRigConnectorElement*> Connectors = InHierarchy->GetElementsOfType<FRigConnectorElement>();
-				for (FRigConnectorElement* Connector : Connectors)
+				for (const FRigConnectorElement* Connector : Connectors)
 				{
 					if (!Redirector.Contains(Connector->GetKey()))
 					{
-						AddNode(Connector->GetKey().ToString());
+						const FString NodeKey = Connector->GetKey().ToString();
+						if(!ContainsElementKeyNode(NodeKey))
+						{
+							AddElementKeyNode(NodeKey);
+						}
 					}
 				}
 			}
 			break;
 		}
+		default:
+		{
+			break;
+		}
 	}
 }
 
@@ -139,34 +222,206 @@ void FControlRigSchematicModel::HandleModularRigModified(EModularRigNotification
 	{
 		case EModularRigNotification::ConnectionChanged:
 		{
-			if (ControlRigBeingDebuggedPtr.IsValid())
+			if (InModule && ControlRigBeingDebuggedPtr.IsValid())
 			{
-				FRigElementKeyRedirector& Redirector = ControlRigBeingDebuggedPtr->GetElementKeyRedirector();
-				for (TPair<FRigElementKey, FRigElementKey>& Pair : Redirector.ExternalKeys)
+				for (const TPair<FRigElementKey, FRigElementKey>& Pair : InModule->Connections)
+				{
+					RemoveElementKeyNode(Pair.Key);
+				}
+				
+				TArray<FRigConnectorElement*> Connectors = ControlRigBeingDebuggedPtr->GetHierarchy()->GetElementsOfType<FRigConnectorElement>();
+				for (const FRigConnectorElement* Connector : Connectors)
 				{
-					RemoveNode(Pair.Key.ToString());
+					if (!InModule->Connections.Contains(Connector->GetKey()))
+					{
+						const FString NodeKey = Connector->GetKey().ToString();
+						if(!ContainsElementKeyNode(NodeKey))
+						{
+							AddElementKeyNode(NodeKey);
+						}
+					}
 				}
 			}
 			break;
 		}
+		default:
+		{
+			break;
+		}
+	}
+}
+
+FVector2d FControlRigSchematicModel::GetPositionForNode(const FSchematicGraphNode* InNode) const
+{
+	if(const FControlRigSchematicRigElementKeyNode* Node = Cast<FControlRigSchematicRigElementKeyNode>(InNode))
+	{
+		if(Node->GetKey().Type != ERigElementType::Connector)
+		{
+			if(const TSharedPtr<FControlRigEditor> Editor = ControlRigEditor.Pin())
+			{
+				if(Editor.IsValid())
+				{
+					if(const UControlRig* ControlRig = Editor->GetControlRig())
+					{
+						if(const URigHierarchy* Hierarchy = ControlRig->GetHierarchy())
+						{
+							const FTransform Transform = Hierarchy->GetGlobalTransform(Node->GetKey());
+							return Editor->ComputePersonaProjectedScreenPos(Transform.GetLocation());
+						}
+					}
+				}
+			}
+		}
+	}
+	return FSchematicGraphModel::GetPositionForNode(InNode);
+}
+
+bool FControlRigSchematicModel::GetPositionAnimationEnabledForNode(const FSchematicGraphNode* InNode) const
+{
+	if(InNode->IsA<FControlRigSchematicRigElementKeyNode>())
+	{
+		return false;
+	}
+	return FSchematicGraphModel::GetPositionAnimationEnabledForNode(InNode);
+}
+
+const FSlateBrush* FControlRigSchematicModel::GetBrushForNode(const FSchematicGraphNode* InNode) const
+{
+	if(const FControlRigSchematicRigElementKeyNode* Node = Cast<FControlRigSchematicRigElementKeyNode>(InNode))
+	{
+		switch(Node->GetKey().Type)
+		{
+			case ERigElementType::Socket:
+			{
+				static const FSlateBrush* UnresolvedSocketBrush = FControlRigEditorStyle::Get().GetBrush( "ControlRig.Schematic.SocketUnresolved");
+				static const FSlateBrush* ResolvedSocketBrush = FControlRigEditorStyle::Get().GetBrush( "ControlRig.Schematic.SocketResolved");
+
+				if (ControlRigBlueprint.IsValid())
+				{
+					for(const TPair<FRigElementKey, FRigElementKey>& Pair : ControlRigBlueprint->ConnectionMap)
+					{
+						if(Pair.Value == Node->GetKey())
+						{
+							return ResolvedSocketBrush;
+						}
+					}
+				}
+
+				return UnresolvedSocketBrush;
+			}
+			case ERigElementType::Connector:
+			{
+				static const FSlateBrush* ConnectorPrimaryBrush = FControlRigEditorStyle::Get().GetBrush( "ControlRig.Schematic.ConnectorPrimary");
+				static const FSlateBrush* ConnectorOptionalBrush = FControlRigEditorStyle::Get().GetBrush( "ControlRig.Schematic.ConnectorOptional");
+				static const FSlateBrush* ConnectorSecondaryBrush = FControlRigEditorStyle::Get().GetBrush( "ControlRig.Schematic.ConnectorSecondary");
+
+				if (ControlRigBlueprint.IsValid())
+				{
+					if(const UControlRig* ControlRig = ControlRigBlueprint->GetDebuggedControlRig())
+					{
+						if(const URigHierarchy* Hierarchy = ControlRig->GetHierarchy())
+						{
+							if(const FRigConnectorElement* Connector = Hierarchy->Find<FRigConnectorElement>(Node->GetKey()))
+							{
+								if(Connector->Settings.Type == EConnectorType::Secondary)
+								{
+									return Connector->Settings.bOptional ? 	ConnectorOptionalBrush : ConnectorSecondaryBrush;
+								}
+							}
+						}
+					}
+				}
+				return ConnectorPrimaryBrush;
+			}
+			case ERigElementType::Bone:
+			{
+				static const FSlateBrush* BoneBrush = FControlRigEditorStyle::Get().GetBrush( "ControlRig.Schematic.Bone");
+				return BoneBrush;
+			}
+			case ERigElementType::Control:
+			{
+				static const FSlateBrush* ControlBrush = FControlRigEditorStyle::Get().GetBrush( "ControlRig.Schematic.Control");
+				return ControlBrush;
+			}
+			case ERigElementType::Null:
+			{
+				static const FSlateBrush* NullBrush = FControlRigEditorStyle::Get().GetBrush( "ControlRig.Schematic.Null");
+				return NullBrush;
+			}
+			default:
+			{
+				break;
+			}
+		}
 	}
+	return FSchematicGraphModel::GetBrushForNode(InNode);
+}
+
+const FText FControlRigSchematicModel::GetToolTipForNode(const FSchematicGraphNode* InNode) const
+{
+	if(const FControlRigSchematicRigElementKeyNode* Node = Cast<FControlRigSchematicRigElementKeyNode>(InNode))
+	{
+		switch(Node->GetKey().Type)
+		{
+			case ERigElementType::Socket:
+			{
+				if (ControlRigBeingDebuggedPtr.IsValid())
+				{
+					if (const URigHierarchy* Hierarchy = ControlRigBeingDebuggedPtr->GetHierarchy())
+					{
+						if(const FRigSocketElement* Socket = Hierarchy->Find<FRigSocketElement>(Node->GetKey()))
+						{
+							const FString Description = Socket->GetDescription(Hierarchy);
+							if(!Description.IsEmpty())
+							{
+								return FText::FromString(Node->GetKey().ToString()+TEXT("\n")+Description);
+							}
+						}
+					}
+				}
+				break;
+			}
+			default:
+			{
+				break;
+			}
+		}
+		return FText::FromString(Node->GetKey().ToString());
+	}
+	return FSchematicGraphModel::GetToolTipForNode(InNode);
+}
+
+ESchematicGraphNodePlacementConstraint FControlRigSchematicModel::GetPlacementForNode(const FSchematicGraphNode* InNode) const
+{
+	if(const FControlRigSchematicRigElementKeyNode* Node = Cast<FControlRigSchematicRigElementKeyNode>(InNode))
+	{
+		if(Node->GetKey().Type == ERigElementType::Connector)
+		{
+			return ESchematicGraphNodePlacementConstraint::BottomRight;
+		}
+	}
+	return FSchematicGraphModel::GetPlacementForNode(InNode);
 }
 
 void FControlRigSchematicModel::HandleSchematicNodeClicked(SSchematicGraphPanel* InPanel, SSchematicGraphNode* InNode)
 {
-	for (FSchematicGraphNode* Node : Nodes)
+	for (const TSharedPtr<FSchematicGraphNode>& Node : Nodes)
 	{
-		Node->bIsSelected = Node == InNode->NodeData;
+		Node->SetSelected(Node->GetGuid() == InNode->GetGuid());
 	}
 
-	if (ControlRigBeingDebuggedPtr.IsValid())
+	if(const FControlRigSchematicRigElementKeyNode* Node =
+			Cast<FControlRigSchematicRigElementKeyNode>(InNode->GetNodeData()))
 	{
-		if (URigHierarchy* Hierarchy = ControlRigBeingDebuggedPtr->GetHierarchy())
+		if (ControlRigBeingDebuggedPtr.IsValid())
 		{
-			if (URigHierarchyController* Controller = Hierarchy->GetController())
+			if (URigHierarchy* Hierarchy = ControlRigBeingDebuggedPtr->GetHierarchy())
 			{
-				TArray<FRigElementKey> Selection = {FRigElementKey(InNode->NodeData->Name)};
-				Controller->SetSelection(Selection);
+				if (URigHierarchyController* Controller = Hierarchy->GetController())
+				{
+					const TArray<FRigElementKey> Selection = {Node->GetKey()};
+					Controller->SetSelection(Selection);
+				}
 			}
 		}
 	}
@@ -184,13 +439,20 @@ void FControlRigSchematicModel::HandleSchematicBeginDrag(SSchematicGraphPanel* I
 		return;
 	}
 
+	const FControlRigSchematicRigElementKeyNode* ElementKeyNode =
+		Cast<FControlRigSchematicRigElementKeyNode>(InNode->GetNodeData());
+	if(ElementKeyNode == nullptr)
+	{
+		return;
+	}
+	
 	URigHierarchy* Hierarchy = ControlRigBeingDebuggedPtr->GetHierarchy();
 	if (!Hierarchy)
 	{
 		return;
 	}
 
-	FRigElementKey DraggedKey(InNode->NodeData->Name);
+	const FRigElementKey DraggedKey = ElementKeyNode->GetKey();
 	if (!DraggedKey.IsValid())
 	{
 		return;
@@ -202,13 +464,13 @@ void FControlRigSchematicModel::HandleSchematicBeginDrag(SSchematicGraphPanel* I
 		return;
 	}
 	
-	FRigConnectorElement* Connector = Cast<FRigConnectorElement>(Element);
+	const FRigConnectorElement* Connector = Cast<FRigConnectorElement>(Element);
 	if (!Connector)
 	{
 		return;
 	}
 
-	UModularRig* ModularRig = Cast<UModularRig>(ControlRigBeingDebuggedPtr);
+	const UModularRig* ModularRig = Cast<UModularRig>(ControlRigBeingDebuggedPtr);
 	if (!ModularRig)
 	{
 		return;
@@ -222,146 +484,62 @@ void FControlRigSchematicModel::HandleSchematicBeginDrag(SSchematicGraphPanel* I
 		return;
 	}
 
-	UModularRigRuleManager* RuleManager = Hierarchy->GetRuleManager();
-	FModularRigResolveResult Result = RuleManager->FindMatches(Connector, ModuleInstance, ControlRigBeingDebuggedPtr->GetElementKeyRedirector());
+	const UModularRigRuleManager* RuleManager = Hierarchy->GetRuleManager();
+	const FModularRigResolveResult Result = RuleManager->FindMatches(Connector, ModuleInstance, ControlRigBeingDebuggedPtr->GetElementKeyRedirector());
 
-	TArray<FRigElementResolveResult> Matches = Result.GetMatches();
-	TArray<FString> NameMatches;
-	Algo::Transform(Matches, NameMatches, [](const FRigElementResolveResult& Match)
-	{
-		return Match.GetKey().ToString();
-	});
+	const TArray<FRigElementResolveResult> Matches = Result.GetMatches();
 
-	// Fade all the unmatched nodes
-	for (FSchematicGraphNode* Node : Nodes)
+	for (const FRigElementResolveResult& Match : Matches)
 	{
-		Node->bFade = !NameMatches.ContainsByPredicate([Node](const FString& Match)
+		// Create a temporary node that will be active only while this drag operation exists
+		if (!ContainsElementKeyNode(Match.GetKey()))
 		{
-			return Node->Name == Match;
-		});
-	};
+			const FSchematicGraphNode* NewNode = AddElementKeyNode(Match.GetKey());
+			TemporaryNodeGuids.Add(NewNode->GetGuid());
+		}
+	}
 
-	// Unfade the ones included in the match
-	for (const FString& Match : NameMatches)
+	// Fade all the unmatched nodes
+	for (const TSharedPtr<FSchematicGraphNode>& Node : Nodes)
 	{
-		bool bExists = Nodes.ContainsByPredicate([Match](const FSchematicGraphNode* Node)
-		{
-			return Match == Node->Name;
-		});
-
-		// Create a temporary node that will be active only while this drag operation exists
-		if (!bExists)
+		if(FControlRigSchematicRigElementKeyNode* ExistingElementKeyNode = Cast<FControlRigSchematicRigElementKeyNode>(Node.Get()))
 		{
-			FSchematicGraphNode* NewNode = AddNode(Match);
-			TemporaryNodes.Add(Match);
+			ExistingElementKeyNode->SetVisibility(Matches.ContainsByPredicate([ExistingElementKeyNode](const FRigElementResolveResult& Match)
+			{
+				return ExistingElementKeyNode->GetKey() == Match.GetKey();
+			}) ? Visible : FadedOut);
 		}
-	}
+	};
 }
 
 void FControlRigSchematicModel::HandleSchematicEndDrag(SSchematicGraphPanel* InPanel, SSchematicGraphNode* InNode, const FDragDropOperation& InDragDropOperation)
 {
-	for (FString& TempNode : TemporaryNodes)
+	for (FGuid& TempNodeGuid : TemporaryNodeGuids)
 	{
-		RemoveNode(TempNode);
+		RemoveNode(TempNodeGuid);
 	}
-	TemporaryNodes.Reset();
+	TemporaryNodeGuids.Reset();
 
-	for (FSchematicGraphNode* Node : Nodes)
+	for (const TSharedPtr<FSchematicGraphNode>& Node : Nodes)
 	{
-		Node->bFade = false;
+		Node->SetVisibility(Visible);
 	}
 }
 
-void FControlRigSchematicModel::HandleUpdateSchematicNodes(SSchematicGraphPanel* InPanel, TSharedPtr<SSchematicGraphNode> InNode)
+void FControlRigSchematicModel::HandleSchematicDrop(SSchematicGraphPanel* InPanel, SSchematicGraphNode* InNode, const FDragDropEvent& InDragDropEvent)
 {
 	if (!ControlRigBlueprint.IsValid())
 	{
 		return;
 	}
 
-	if (InPanel->bIsOverlay)
-	{
-		FTransform Transform;
-		if (ControlRigBeingDebuggedPtr.IsValid())
-		{
-			URigHierarchy* Hierarchy = ControlRigBeingDebuggedPtr->GetHierarchy();
-			FRigElementKey ElementKey(InNode->NodeData->Name);
-
-			if (ElementKey.IsTypeOf(ERigElementType::Socket))
-			{
-				Transform = Hierarchy->GetGlobalTransform(ElementKey);
-				FVector2D PixelPos = ControlRigEditor.Pin()->ComputePersonaProjectedScreenPos(Transform.GetLocation());
-				InNode->SetPosition(PixelPos, true);
-
-				if (const FRigElementKey* ConnectorKey = ControlRigBeingDebuggedPtr->GetElementKeyRedirector().FindReverse(ElementKey))
-				{
-					InNode->Brush = *FControlRigEditorStyle::Get().GetBrush( "ControlRig.Schematic.SocketResolved");
-					InNode->Size->Set(InNode->OriginalSize*0.5);
-				}
-				else
-				{
-					InNode->Brush = *FControlRigEditorStyle::Get().GetBrush( "ControlRig.Schematic.SocketUnresolved");
-					InNode->Size->Set(InNode->OriginalSize);
-				}
-
-				if (FRigSocketElement* SocketElement = Cast<FRigSocketElement>(Hierarchy->Find(ElementKey)))
-				{
-					InNode->Brush.TintColor = SocketElement->GetColor(Hierarchy);
-				}
-				else
-				{
-					InNode->Brush.TintColor = FStyleColors::AccentBlue;
-				}
-			}
-			else if(ElementKey.IsTypeOf(ERigElementType::Connector))
-			{
-				InNode->NodeData->Placement = ESchematicGraphNodePlacement::BottomRight;
-				InNode->Size->Set(InNode->OriginalSize);
-				if (FRigConnectorElement* ConnectorElement = Cast<FRigConnectorElement>(Hierarchy->Find(ElementKey)))
-				{
-					if (ConnectorElement->Settings.Type == EConnectorType::Primary)
-					{
-						InNode->Brush = *FControlRigEditorStyle::Get().GetBrush( "ControlRig.Schematic.ConnectorPrimary");
-					}
-					else if(ConnectorElement->Settings.Type == EConnectorType::Secondary)
-					{
-						if (ConnectorElement->Settings.bOptional)
-						{
-							InNode->Brush = *FControlRigEditorStyle::Get().GetBrush( "ControlRig.Schematic.ConnectorOptional");
-						}
-						else
-						{
-							InNode->Brush = *FControlRigEditorStyle::Get().GetBrush( "ControlRig.Schematic.ConnectorSecondary");
-						}
-					}
-				}
-			}
-			else
-			{
-				Transform = Hierarchy->GetGlobalTransform(ElementKey);
-				FVector2D PixelPos = ControlRigEditor.Pin()->ComputePersonaProjectedScreenPos(Transform.GetLocation());
-				InNode->SetPosition(PixelPos, true);
-
-				InNode->Size->Set(FVector2d(20, 20));
-				InNode->Brush.TintColor = FStyleColors::AccentBlue;
-			}
-		}
-
-	}
-	else
-	{
-		InNode->SetPosition(FVector2D(0, 0));
-		InNode->Brush = *FAppStyle::GetBrush("WhiteTexture");
-	}
-}
-
-void FControlRigSchematicModel::HandleSchematicDrop(SSchematicGraphPanel* InPanel, SSchematicGraphNode* InNode, const FDragDropEvent& InDragDropEvent)
-{
-	if (!ControlRigBlueprint.IsValid())
+	const FControlRigSchematicRigElementKeyNode* ElementKeyNode =
+	Cast<FControlRigSchematicRigElementKeyNode>(InNode->GetNodeData());
+	if(ElementKeyNode == nullptr)
 	{
 		return;
 	}
-
+	
 	UControlRig* ControlRig = ControlRigBlueprint->GetDebuggedControlRig();
 	if (!ControlRig)
 	{
@@ -374,7 +552,7 @@ void FControlRigSchematicModel::HandleSchematicDrop(SSchematicGraphPanel* InPane
 		return;
 	}
 	
-	FRigBaseElement* Target = Hierarchy->Find(FRigElementKey(InNode->NodeData->Name));
+	FRigBaseElement* Target = Hierarchy->Find(ElementKeyNode->GetKey());
 	if (!Target)
 	{
 		return;
@@ -415,6 +593,12 @@ void FControlRigSchematicModel::HandleSchematicDrop(SSchematicGraphPanel* InPane
 								}
 							}
 						}
+
+						const FName TargetModulePath = Hierarchy->GetNameMetadata(Target->GetKey(), URigHierarchy::NameSpaceMetadataName, NAME_None);
+						if(!TargetModulePath.IsNone())
+						{
+							(void)Controller->ReparentModule(ModulePath, TargetModulePath.ToString());
+						}
 						Controller->ConnectConnectorToElement(PrimaryConnectorKey, Target->GetKey());
 					}
 				}
@@ -425,15 +609,15 @@ void FControlRigSchematicModel::HandleSchematicDrop(SSchematicGraphPanel* InPane
 	{
 		if (UModularRigController* Controller = ControlRigBlueprint->GetModularRigController())
 		{
-			TArray<FString> Sources = SchematicDragDropOp->GetElements();
-			for (FString& SourceName : Sources)
+			const TArray<FGuid> Sources = SchematicDragDropOp->GetElements();
+			for (const TPair<FRigElementKey, FGuid> Pair : RigElementKeyToGuid)
 			{
-				FRigElementKey Key(SourceName);
-				if (FRigBaseElement* SourceElement = Hierarchy->Find(SourceName))
+				if(Sources.Contains(Pair.Value))
 				{
-					if (FRigConnectorElement* Connector = Cast<FRigConnectorElement>(SourceElement))
+					if (Hierarchy->Find<FRigConnectorElement>(Pair.Key))
 					{
-						Controller->ConnectConnectorToElement(Key, Target->GetKey());
+						Controller->ConnectConnectorToElement(Pair.Key, Target->GetKey());
+						break;
 					}
 				}
 			}
diff --git a/Engine/Plugins/Animation/ControlRig/Source/ControlRigEditor/Private/Editor/ControlRigSchematicModel.h b/Engine/Plugins/Animation/ControlRig/Source/ControlRigEditor/Private/Editor/ControlRigSchematicModel.h
--- a/Engine/Plugins/Animation/ControlRig/Source/ControlRigEditor/Private/Editor/ControlRigSchematicModel.h
+++ b/Engine/Plugins/Animation/ControlRig/Source/ControlRigEditor/Private/Editor/ControlRigSchematicModel.h
@@ -12,28 +12,63 @@ class UControlRig;
 class URigHierarchy;
 struct FRigBaseElement;
 
+/** Node for the schematic view */
+class FControlRigSchematicRigElementKeyNode : public FSchematicGraphNode
+{
+public:
+
+	SCHEMATICGRAPHNODE_BODY(FControlRigSchematicRigElementKeyNode, FSchematicGraphNode)
+
+	virtual ~FControlRigSchematicRigElementKeyNode() override {}
+	
+	const FRigElementKey& GetKey() const { return Key; }
+	virtual FString GetDragDropDecoratorLabel() const override;
+	
+protected:
+
+	FRigElementKey Key;
+
+	friend class FControlRigSchematicModel;
+};
+
 /** Model for the schematic views */
-struct FControlRigSchematicModel : FSchematicGraph
+class FControlRigSchematicModel : public FSchematicGraphModel
 {
-	~FControlRigSchematicModel();
+public:
+	virtual ~FControlRigSchematicModel() override;
 	
 	void SetEditor(const TSharedRef<FControlRigEditor>& InEditor);
 
+	virtual void Reset() override;
+	FControlRigSchematicRigElementKeyNode* AddElementKeyNode(const FRigElementKey& InKey, bool bNotify = true);
+	const FControlRigSchematicRigElementKeyNode* FindElementKeyNode(const FRigElementKey& InKey) const;
+	bool ContainsElementKeyNode(const FRigElementKey& InKey) const;
+	virtual bool RemoveNode(const FGuid& InGuid) override;
+	bool RemoveElementKeyNode(const FRigElementKey& InKey);
+
 	void OnSetObjectBeingDebugged(UObject* InObject);
 	void OnHierarchyModified(ERigHierarchyNotification InNotif, URigHierarchy* InHierarchy, const FRigBaseElement* InElement);
 	void HandleModularRigModified(EModularRigNotification InNotification, const FRigModuleReference* InModule);
 
-	void HandleUpdateSchematicNodes(SSchematicGraphPanel* InPanel, TSharedPtr<SSchematicGraphNode> InNode);
+	virtual FVector2d GetPositionForNode(const FSchematicGraphNode* InNode) const override;
+	virtual bool GetPositionAnimationEnabledForNode(const FSchematicGraphNode* InNode) const override;
+	virtual const FSlateBrush* GetBrushForNode(const FSchematicGraphNode* InNode) const override;
+	virtual const FText GetToolTipForNode(const FSchematicGraphNode* InNode) const override;
+	virtual ESchematicGraphNodePlacementConstraint GetPlacementForNode(const FSchematicGraphNode* InNode) const override;
+
+private:
+
 	void HandleSchematicNodeClicked(SSchematicGraphPanel* InPanel, SSchematicGraphNode* InNode);
 	void HandleSchematicBeginDrag(SSchematicGraphPanel* InPanel, SSchematicGraphNode* InNode, const FDragDropOperation& InDragDropOperation);
 	void HandleSchematicEndDrag(SSchematicGraphPanel* InPanel, SSchematicGraphNode* InNode, const FDragDropOperation& InDragDropOperation);
 	void HandleSchematicDrop(SSchematicGraphPanel* InPanel, SSchematicGraphNode* InNode, const FDragDropEvent& InDragDropEvent);
-
-private:
 	
 	TWeakPtr<FControlRigEditor> ControlRigEditor;
 	TWeakObjectPtr<UControlRigBlueprint> ControlRigBlueprint;
 	TWeakObjectPtr<UControlRig> ControlRigBeingDebuggedPtr;
 
-	TArray<FString> TemporaryNodes;
+	TArray<FGuid> TemporaryNodeGuids;
+	TMap<FRigElementKey, FGuid> RigElementKeyToGuid;
+
+	friend class FControlRigEditor;
 };
\ No newline at end of file
diff --git a/Engine/Source/Developer/AnimationWidgets/Private/SchematicGraphPanel/SSchematicGraphPanel.cpp b/Engine/Source/Developer/AnimationWidgets/Private/SchematicGraphPanel/SSchematicGraphPanel.cpp
--- a/Engine/Source/Developer/AnimationWidgets/Private/SchematicGraphPanel/SSchematicGraphPanel.cpp
+++ b/Engine/Source/Developer/AnimationWidgets/Private/SchematicGraphPanel/SSchematicGraphPanel.cpp
@@ -7,8 +7,22 @@
 
 #define LOCTEXT_NAMESPACE "SSchematicGraphPanel"
 
+const FName& FSchematicGraphNode::GetType() const
+{
+	return FSchematicGraphNode::Type;
+}
 
-void FSchematicGraph::Reset()
+bool FSchematicGraphNode::IsA(const FName& InType) const
+{
+	return Type == InType;
+}
+
+FString FSchematicGraphNode::GetDragDropDecoratorLabel() const
+{
+	return GetGuid().ToString();
+}
+
+void FSchematicGraphModel::Reset()
 {
 	Nodes.Reset();
 	Links.Reset();
@@ -19,69 +33,184 @@ void FSchematicGraph::Reset()
 	}
 }
 
-FSchematicGraphNode* FSchematicGraph::AddNode(const FString& InName)
+bool FSchematicGraphModel::RemoveNode(const FGuid& InGuid)
 {
-	if (Nodes.ContainsByPredicate([InName](const FSchematicGraphNode* Node)
+	if(const FSchematicGraphNode* Node = FindNode(InGuid))
 	{
-		return Node->Name == InName;
-	}))
+		if (OnNodeRemovedDelegate.IsBound())
+		{
+			OnNodeRemovedDelegate.Broadcast(Node);
+		}
+		const FGuid Guid = Node->GetGuid();
+		NodeByGuid.Remove(Node->GetGuid());
+		Nodes.RemoveAll([Guid](const TSharedPtr<FSchematicGraphNode>& ExistingNode) -> bool
+		{
+			return ExistingNode->GetGuid() == Guid;
+		});
+		return true;
+	}
+	return false;
+}
+
+FVector2d FSchematicGraphModel::GetPositionForNode(const FGuid& InGuid) const
+{
+	if(const FSchematicGraphNode* Node = FindNode(InGuid))
 	{
-		return nullptr;
+		return GetPositionForNode(Node);
 	}
+	return FVector2d::ZeroVector;
+}
 
-	FSchematicGraphNode* NewElement = static_cast<FSchematicGraphNode*>(FMemory::Malloc(sizeof(FSchematicGraphNode)));
-	new (NewElement) FSchematicGraphNode();
-	NewElement->Name = InName;
-	Nodes.Add(NewElement);
+FVector2d FSchematicGraphModel::GetPositionForNode(const FSchematicGraphNode* InNode) const
+{
+	check(InNode);
+	return InNode->GetPosition();
+}
 
-	if (OnNodeAddedDelegate.IsBound())
+FVector2d FSchematicGraphModel::GetPositionOffsetForNode(const FGuid& InGuid) const
+{
+	if(const FSchematicGraphNode* Node = FindNode(InGuid))
 	{
-		OnNodeAddedDelegate.Broadcast(NewElement);
+		return GetPositionOffsetForNode(Node);
 	}
-	return NewElement;
+	return FVector2d::ZeroVector;
+}
+
+FVector2d FSchematicGraphModel::GetPositionOffsetForNode(const FSchematicGraphNode* InNode) const
+{
+	check(InNode);
+	return InNode->GetPositionOffset();
 }
 
-bool FSchematicGraph::RenameNode(const FString& InOldName, const FString& InNewName)
+bool FSchematicGraphModel::GetPositionAnimationEnabledForNode(const FGuid& InNodeGuid) const
 {
-	FSchematicGraphNode** FoundNode = Nodes.FindByPredicate([InOldName](const FSchematicGraphNode* Node)
+	if(const FSchematicGraphNode* Node = FindNode(InNodeGuid))
 	{
-		return Node->Name == InOldName;
-	});
+		return GetPositionAnimationEnabledForNode(Node);
+	}
+	return false;
+}
+
+bool FSchematicGraphModel::GetPositionAnimationEnabledForNode(const FSchematicGraphNode* InNode) const
+{
+	return true;
+}
+
+FVector2d FSchematicGraphModel::GetSizeForNode(const FGuid& InNodeGuid) const
+{
+	if(const FSchematicGraphNode* Node = FindNode(InNodeGuid))
+	{
+		return GetSizeForNode(Node);
+	}
+	return FVector2d::ZeroVector;
+}
+
+FVector2d FSchematicGraphModel::GetSizeForNode(const FSchematicGraphNode* InNode) const
+{
+	static const FVector2d DefaultSize = FVector2d(50.0,50.0);
+	return DefaultSize;
+}
+
+float FSchematicGraphModel::GetScaleForNode(const FGuid& InNodeGuid) const
+{
+	if(const FSchematicGraphNode* Node = FindNode(InNodeGuid))
+	{
+		return GetScaleForNode(Node);
+	}
+	return 1.f;
+}
+
+float FSchematicGraphModel::GetScaleForNode(const FSchematicGraphNode* InNode) const
+{
+	return 1.f;
+}
 
-	if (FoundNode)
+FLinearColor FSchematicGraphModel::GetColorForNode(const FGuid& InGuid) const
+{
+	if(const FSchematicGraphNode* Node = FindNode(InGuid))
 	{
-		if (OnNodeRenamedDelegate.IsBound())
+		const FLinearColor Color = GetColorForNode(Node);
+		if(GetVisibilityForNode(Node) == ESchematicGraphNodeVisibility::FadedOut)
 		{
-			OnNodeRenamedDelegate.Broadcast(*FoundNode);
+			return Color * 0.5f;
 		}
-		(*FoundNode)->Name = InNewName;
-		return true;
+		return Color;
 	}
-	
-	return false;
+	return FLinearColor::White;
 }
 
-bool FSchematicGraph::RemoveNode(const FString& InName)
+FLinearColor FSchematicGraphModel::GetColorForNode(const FSchematicGraphNode* InNode) const
 {
-	FSchematicGraphNode** FoundNodePtr = Nodes.FindByPredicate([InName](const FSchematicGraphNode* Node)
-	{
-		return Node->Name == InName;
-	});
-	if (FoundNodePtr)
+	check(InNode);
+	return InNode->GetColor();
+}
+
+const FSlateBrush* FSchematicGraphModel::GetBrushForNode(const FGuid& InGuid) const
+{
+	if(const FSchematicGraphNode* Node = FindNode(InGuid))
 	{
-		FSchematicGraphNode* FoundNode = *FoundNodePtr;
-		if (OnNodeRemovedDelegate.IsBound())
+		if(const FSlateBrush* Brush = GetBrushForNode(Node))
 		{
-			OnNodeRemovedDelegate.Broadcast(FoundNode);
+			return Brush;
 		}
-		Nodes.Remove(FoundNode);
-		FMemory::Free(FoundNode);
-		return true;
 	}
-	return false;
+
+	static const FSlateBrush* DefaultBrush = SSchematicGraphNode::FArguments()._Brush.Get();
+	return DefaultBrush;
+}
+
+const FSlateBrush* FSchematicGraphModel::GetBrushForNode(const FSchematicGraphNode* InNode) const
+{
+	check(InNode);
+	return InNode->GetBrush();
+}
+
+const FText FSchematicGraphModel::GetToolTipForNode(const FGuid& InGuid) const
+{
+	if(const FSchematicGraphNode* Node = FindNode(InGuid))
+	{
+		return GetToolTipForNode(Node);
+	}
+	return FText();
+}
+
+const FText FSchematicGraphModel::GetToolTipForNode(const FSchematicGraphNode* InNode) const
+{
+	check(InNode);
+	return InNode->GetToolTip();
+}
+
+ESchematicGraphNodePlacementConstraint FSchematicGraphModel::GetPlacementForNode(const FGuid& InGuid) const
+{
+	if(const FSchematicGraphNode* Node = FindNode(InGuid))
+	{
+		return GetPlacementForNode(Node);
+	}
+	return ESchematicGraphNodePlacementConstraint::Free;
+}
+
+ESchematicGraphNodePlacementConstraint FSchematicGraphModel::GetPlacementForNode(const FSchematicGraphNode* InNode) const
+{
+	check(InNode);
+	return InNode->GetPlacement();
+}
+
+ESchematicGraphNodeVisibility FSchematicGraphModel::GetVisibilityForNode(const FGuid& InGuid) const
+{
+	if(const FSchematicGraphNode* Node = FindNode(InGuid))
+	{
+		return GetVisibilityForNode(Node);
+	}
+	return ESchematicGraphNodeVisibility::Visible;
+}
+
+ESchematicGraphNodeVisibility FSchematicGraphModel::GetVisibilityForNode(const FSchematicGraphNode* InNode) const
+{
+	check(InNode);
+	return InNode->GetVisibility();
 }
 
-TSharedRef<FSchematicGraphNodeDragDropOp> FSchematicGraphNodeDragDropOp::New(TArray<SSchematicGraphNode*> InSchematicGraphNodes, const TArray<FString>& InElements, FSchematicGraphNodeDragDropOp::FOnEndDrag InOnEndDragDelegate)
+TSharedRef<FSchematicGraphNodeDragDropOp> FSchematicGraphNodeDragDropOp::New(TArray<SSchematicGraphNode*> InSchematicGraphNodes, const TArray<FGuid>& InElements, FSchematicGraphNodeDragDropOp::FOnEndDrag InOnEndDragDelegate)
 {
 	TSharedRef<FSchematicGraphNodeDragDropOp> Operation = MakeShared<FSchematicGraphNodeDragDropOp>();
 	Operation->SchematicGraphNodes = InSchematicGraphNodes;
@@ -91,6 +220,17 @@ TSharedRef<FSchematicGraphNodeDragDropOp> FSchematicGraphNodeDragDropOp::New(TAr
 	return Operation;
 }
 
+FSchematicGraphNodeDragDropOp::~FSchematicGraphNodeDragDropOp()
+{
+	if (OnEndDragDelegate.IsBound())
+	{
+		for (SSchematicGraphNode* Node : SchematicGraphNodes)
+		{
+			OnEndDragDelegate.Execute(Node, *this);
+		}
+	}
+}
+
 TSharedPtr<SWidget> FSchematicGraphNodeDragDropOp::GetDefaultDecorator() const
 {
 	return SNew(SBorder)
@@ -98,35 +238,58 @@ TSharedPtr<SWidget> FSchematicGraphNodeDragDropOp::GetDefaultDecorator() const
 		.BorderImage(FAppStyle::GetBrush("Menu.Background"))
 		[
 			SNew(STextBlock)
-			.Text(FText::FromString(GetJoinedElementNames()))
+			.Text(FText::FromString(GetJoinedDecoratorLabels()))
 		];
 }
 
-FString FSchematicGraphNodeDragDropOp::GetJoinedElementNames() const
+FString FSchematicGraphNodeDragDropOp::GetJoinedDecoratorLabels() const
 {
-	TArray<FString> ElementNameStrings;
-	for (const FString& Element: Elements)
+	TArray<FString> DecoratorLabels;
+	for(const SSchematicGraphNode* GraphNode : SchematicGraphNodes)
 	{
-		ElementNameStrings.Add(Element);
+		if(const FSchematicGraphNode* Node = GraphNode->GetNodeData())
+		{
+			DecoratorLabels.Add(Node->GetDragDropDecoratorLabel());
+		}
 	}
-	return FString::Join(ElementNameStrings, TEXT(","));
+	return FString::Join(DecoratorLabels, TEXT(","));
+}
+
+SSchematicGraphNode::FArguments::FArguments()
+: _NodeData(nullptr)
+, _Brush(nullptr)
+{
+	static const FSchematicGraphNode EmptyNodeData;
+	_NodeData = &EmptyNodeData;
+	static const FSlateBrush* WhiteTexture = FAppStyle::GetBrush("WhiteTexture");
+	_Brush = WhiteTexture;
 }
 
 void SSchematicGraphNode::Construct(const FArguments& InArgs)
 {
-	NodeData = InArgs._NodeData;
+	if(InArgs._NodeData)
+	{
+		NodeData = InArgs._NodeData;
+	}
 	OnClickedDelegate = InArgs._OnClicked;
 	OnBeginDragDelegate = InArgs._OnBeginDrag;
 	OnEndDragDelegate = InArgs._OnEndDrag;
 	OnDropDelegate = InArgs._OnDrop;
-	TEasingAttributeInterpolator<FVector2d>::FSettings Vector2DInterpSettings(EEasingInterpolatorType::BounceEaseInOut, 0.2f);
-	Position = TAnimatedAttribute<FVector2d>::Create(Vector2DInterpSettings, FVector2d::ZeroVector);
-	Size = TAnimatedAttribute<FVector2d>::Create(Vector2DInterpSettings, OriginalSize);
-	TEasingAttributeInterpolator<float>::FSettings FloatInterpSettings(EEasingInterpolatorType::BounceEaseInOut, 0.2f);
-	Scale = TAnimatedAttribute<float>::Create(FloatInterpSettings, 1.f);
-	Brush = *FAppStyle::GetBrush("WhiteTexture");
-	Brush.TintColor = EStyleColor::AccentWhite;
-	SetToolTipText(FText::FromString(NodeData->Name));
+
+	Position = InArgs._Position;
+	Size = InArgs._Size;
+	Scale = InArgs._Scale;
+	Color = InArgs._Color;
+	if(InArgs._Brush.IsSet() || InArgs._Brush.IsBound())
+	{
+		Brush = InArgs._Brush;
+	}
+	OriginalSize = Size->Get();
+
+	if(InArgs._ToolTipText.IsBound() || InArgs._ToolTipText.IsSet())
+	{
+		SetToolTipText(InArgs._ToolTipText);
+	}
 }
 
 FVector2D SSchematicGraphNode::ComputeDesiredSize(float LayoutScaleMultiplier) const
@@ -139,23 +302,16 @@ int32 SSchematicGraphNode::OnPaint(const FPaintArgs& Args, const FGeometry& Allo
 	int32 NewLayerId = SNodePanel::SNode::OnPaint(Args, AllottedGeometry, MyCullingRect, OutDrawElements, LayerId, InWidgetStyle, bParentEnabled);
 	NewLayerId++;
 
-	FVector2d CurSize = Size->Get() * Scale->Get();
-	FVector2d SizeOffset = (CurSize-OriginalSize)*-0.5;
-	FVector2d TotalOffset = SizeOffset;
-
-	FLinearColor TintColor = Brush.TintColor.GetSpecifiedColor();
-	if (NodeData->bFade)
-	{
-		TintColor.A = 0.3;
-	}
+	const FVector2d CurSize = Size->Get() * Scale->Get();
+	const FVector2d SizeOffset = (CurSize-OriginalSize)*-0.5;
 
 	FSlateDrawElement::MakeBox(
 				OutDrawElements,
 				NewLayerId,
-				AllottedGeometry.ToPaintGeometry(CurSize, FSlateLayoutTransform(TotalOffset)),
-				&Brush,
+				AllottedGeometry.ToPaintGeometry(CurSize, FSlateLayoutTransform(SizeOffset)),
+				Brush.Get(),
 				ESlateDrawEffect::None,
-				TintColor
+				Color.IsValid() ? Color->Get() : FLinearColor::White
 				);
 	
 	return NewLayerId;
@@ -163,18 +319,30 @@ int32 SSchematicGraphNode::OnPaint(const FPaintArgs& Args, const FGeometry& Allo
 
 void SSchematicGraphNode::OnMouseEnter(const FGeometry& MyGeometry, const FPointerEvent& MouseEvent)
 {
+	if(!IsInteractive())
+	{
+		return;
+	}
 	SNode::OnMouseEnter(MyGeometry, MouseEvent);
 	Scale->Set(ScaledUp);
 }
 
 void SSchematicGraphNode::OnMouseLeave(const FPointerEvent& MouseEvent)
 {
+	if(!IsInteractive())
+	{
+		return;
+	}
 	SNode::OnMouseLeave(MouseEvent);
 	Scale->Set(1.f);
 }
 
 FReply SSchematicGraphNode::OnDragOver(const FGeometry& MyGeometry, const FDragDropEvent& DragDropEvent)
 {
+	if(!IsInteractive())
+	{
+		return FReply::Unhandled();
+	}
 	SNode::OnDragOver(MyGeometry, DragDropEvent);
 	Scale->Set(ScaledDown);
 	return FReply::Handled();
@@ -182,21 +350,30 @@ FReply SSchematicGraphNode::OnDragOver(const FGeometry& MyGeometry, const FDragD
 
 void SSchematicGraphNode::OnDragLeave(const FDragDropEvent& DragDropEvent)
 {
+	if(!IsInteractive())
+	{
+		return;
+	}
 	SNode::OnDragLeave(DragDropEvent);
 	Scale->Set(1.f);
 }
 
 FReply SSchematicGraphNode::OnDrop(const FGeometry& MyGeometry, const FDragDropEvent& DragDropEvent)
 {
+	if(!IsInteractive())
+	{
+		return FReply::Unhandled();
+	}
+
 	// Avoid dropping onto itself
-	TSharedPtr<FSchematicGraphNodeDragDropOp> SchematicDragDropOp = DragDropEvent.GetOperationAs<FSchematicGraphNodeDragDropOp>();
+	const TSharedPtr<FSchematicGraphNodeDragDropOp> SchematicDragDropOp = DragDropEvent.GetOperationAs<FSchematicGraphNodeDragDropOp>();
 	if (SchematicDragDropOp)
 	{
 		if (!SchematicDragDropOp->GetElements().IsEmpty())
 		{
-			if (SchematicDragDropOp->GetElements().ContainsByPredicate([this](const FString& Name)
+			if (SchematicDragDropOp->GetElements().ContainsByPredicate([this](const FGuid& Guid)
 			{
-				return Name == NodeData->Name;
+				return Guid == NodeData->GetGuid();
 			}))
 			{
 				return FReply::Unhandled();
@@ -212,11 +389,16 @@ FReply SSchematicGraphNode::OnDrop(const FGeometry& MyGeometry, const FDragDropE
 
 FReply SSchematicGraphNode::OnDragDetected(const FGeometry& MyGeometry, const FPointerEvent& MouseEvent)
 {
-	TArray<FString> DraggedElements = {NodeData->Name};
+	TArray<FGuid> DraggedElements = {GetGuid()};
 	if (MouseEvent.IsMouseButtonDown(EKeys::LeftMouseButton) && DraggedElements.Num() > 0)
 	{
 		bIsBeingDragged = true;
-		TSharedRef<FSchematicGraphNodeDragDropOp> DragDropOp = FSchematicGraphNodeDragDropOp::New({this}, MoveTemp(DraggedElements), OnEndDragDelegate);
+
+		const FVector2f AbsoluteMousePosition = MouseEvent.GetScreenSpacePosition();
+		const FVector2d LocalMousePosition =  (AbsoluteMousePosition - MyGeometry.GetAbsolutePosition()) / MyGeometry.GetAccumulatedLayoutTransform().GetScale();
+		OffsetDuringDrag = -LocalMousePosition;
+		
+		const TSharedRef<FSchematicGraphNodeDragDropOp> DragDropOp = FSchematicGraphNodeDragDropOp::New({this}, MoveTemp(DraggedElements), OnEndDragDelegate);
 		OnBeginDragDelegate.ExecuteIfBound(this, DragDropOp.Get());
 		return FReply::Handled().BeginDragDrop(DragDropOp);
 	}
@@ -239,20 +421,31 @@ FVector2d SSchematicGraphNode::GetPosition() const
 	return Position->Get() - (OriginalSize*0.5);
 }
 
-void SSchematicGraphNode::SetPosition(const FVector2d& InPosition, bool bImmediately)
+void SSchematicGraphNode::EnablePositionAnimation(bool bEnabled)
 {
-	bImmediately ?
-		Position->SetValueAndStop(FVector2d(InPosition))
-		: Position->Set(InPosition);
+	Position->EnableInterpolation(bEnabled);
 }
 
-void SSchematicGraphPanel::SetSchematicGraph(FSchematicGraph* InGraphData)
+const FGuid SSchematicGraphNode::GetGuid() const
+{
+	return GetNodeData()->GetGuid();
+}
+
+bool SSchematicGraphNode::IsInteractive() const
+{
+	if(NodeData)
+	{
+		return NodeData->GetVisibility() == ESchematicGraphNodeVisibility::Visible;
+	}
+	return true;
+}
+
+void SSchematicGraphPanel::SetSchematicGraph(FSchematicGraphModel* InGraphData)
 {
 	if (GraphData)
 	{
 		GraphData->OnNodeAdded().RemoveAll(this);
 		GraphData->OnNodeRemoved().RemoveAll(this);
-		GraphData->OnNodeRenamed().RemoveAll(this);
 		GraphData->OnGraphReset().RemoveAll(this);
 	}
 	
@@ -261,7 +454,6 @@ void SSchematicGraphPanel::SetSchematicGraph(FSchematicGraph* InGraphData)
 	{
 		GraphData->OnNodeAdded().AddSP(this, &SSchematicGraphPanel::AddNode);
 		GraphData->OnNodeRemoved().AddSP(this, &SSchematicGraphPanel::RemoveNode);
-		GraphData->OnNodeRenamed().AddSP(this, &SSchematicGraphPanel::RenameNode);
 		GraphData->OnGraphReset().AddSP(this, &SSchematicGraphPanel::RebuildPanel);
 	}
 }
@@ -275,14 +467,16 @@ void SSchematicGraphPanel::Construct(const FArguments& InArgs)
 	PaddingTop = InArgs._PaddingTop;
 	PaddingBottom = InArgs._PaddingBottom;
 	PaddingInterNode = InArgs._PaddingInterNode;
-	UpdateNodeWidgetDelegate = InArgs._OnUpdateNodeWidget;
 	OnNodeClickedDelegate = InArgs._OnNodeClicked;
 	OnBeginDragDelegate = InArgs._OnBeginDrag;
 	OnEndDragDelegate = InArgs._OnEndDrag;
 	OnDropDelegate = InArgs._OnDrop;
 
-	TEasingAttributeInterpolator<float>::FSettings FloatInterpSettings(EEasingInterpolatorType::BounceEaseInOut, 0.2f);
-	FadeBackgroundAlpha = TAnimatedAttribute<float>::Create(FloatInterpSettings, 0.f);
+	TEasingAttributeInterpolator<float>::FSettings FloatInterpSettings(EEasingInterpolatorType::CubicEaseOut, 0.35f);
+	FadeBackgroundAlpha = FFloatAttribute::CreateWithGetter(FloatInterpSettings, FFloatAttribute::FGetter::CreateLambda([this]() -> float
+	{
+		return bIsDragDropping ? 0.5f : 0.f;
+	}));
 	
 	SNodePanel::Construct();
 
@@ -291,7 +485,6 @@ void SSchematicGraphPanel::Construct(const FArguments& InArgs)
 	{
 		GraphData->OnNodeAdded().AddSP(this, &SSchematicGraphPanel::AddNode);
 		GraphData->OnNodeRemoved().AddSP(this, &SSchematicGraphPanel::RemoveNode);
-		GraphData->OnNodeRenamed().AddSP(this, &SSchematicGraphPanel::RenameNode);
 		GraphData->OnGraphReset().AddSP(this, &SSchematicGraphPanel::RebuildPanel);
 	}
 }
@@ -302,36 +495,50 @@ void SSchematicGraphPanel::RebuildPanel()
 
 	if (GraphData)
 	{
-		for (FSchematicGraphNode* Node : GraphData->Nodes)
+		for (const TSharedPtr<FSchematicGraphNode>& Node : GraphData->GetNodes())
 		{
-			AddNode(Node);
+			AddNode(Node.Get());
 		}
 	}
 }
 
-void SSchematicGraphPanel::AddNode(FSchematicGraphNode* NodeToAdd)
+void SSchematicGraphPanel::AddNode(const FSchematicGraphNode* InNodeToAdd)
 {
-	TSharedRef<SSchematicGraphNode> NewNode = SNew(SSchematicGraphNode)
+	TEasingAttributeInterpolator<FVector2d>::FSettings Vector2DInterpSettings(EEasingInterpolatorType::CubicEaseOut, 0.2f);
+	TEasingAttributeInterpolator<float>::FSettings FloatInterpSettings(EEasingInterpolatorType::CubicEaseOut, 0.2f);
+	TEasingAttributeInterpolator<FLinearColor>::FSettings ColorInterpSettings(EEasingInterpolatorType::CubicEaseOut, 0.2f);
+
+	const FGuid Guid = InNodeToAdd->GetGuid();
+	const auto Position = FVector2dAttribute::CreateWithGetter(Vector2DInterpSettings, FVector2dAttribute::FGetter::CreateSP(this, &SSchematicGraphPanel::GetPositionForNode, Guid));
+	const auto Size = FVector2dAttribute::CreateWithGetter(Vector2DInterpSettings, FVector2dAttribute::FGetter::CreateSP(this, &SSchematicGraphPanel::GetSizeForNode, Guid));
+	const auto Scale = FFloatAttribute::CreateWithGetter(FloatInterpSettings, FFloatAttribute::FGetter::CreateSP(this, &SSchematicGraphPanel::GetScaleForNode, Guid), 0.f);
+	const auto Color = FLinearColorAttribute::CreateWithGetter(ColorInterpSettings, FLinearColorAttribute::FGetter::CreateSP(this, &SSchematicGraphPanel::GetColorForNode, Guid));
+
+	const TSharedRef<SSchematicGraphNode> NewNode = SNew(SSchematicGraphNode)
+														.Position(Position)
+														.Size(Size)
+														.Scale(Scale)
+														.Color(Color)
+														.Brush(this, &SSchematicGraphPanel::GetBrushForNode, InNodeToAdd->GetGuid())
+														.ToolTipText(this, &SSchematicGraphPanel::GetToolTipForNode, InNodeToAdd->GetGuid())
 														.OnClicked_Raw(this, &SSchematicGraphPanel::OnNodeClicked)
 														.OnBeginDrag_Raw(this, &SSchematicGraphPanel::OnBeginDragEvent)
 														.OnEndDrag_Raw(this, &SSchematicGraphPanel::OnEndDragEvent)
 														.OnDrop_Raw(this, &SSchematicGraphPanel::OnDropEvent)
-														.NodeData(NodeToAdd);
+														.NodeData(InNodeToAdd);
 	SNodePanel::AddGraphNode(NewNode);
-
-	// Update the position/size/brush immediately without animation
-	{
-		TGuardValue<bool> AnimateNodeGuard(bAnimatePosition, false);
-		Tick(0.f);
-	}
+	NodeByGuid.Add(Guid, NewNode.ToSharedPtr());
 }
 
-void SSchematicGraphPanel::RemoveNode(FSchematicGraphNode* InNodeToRemove)
+void SSchematicGraphPanel::RemoveNode(const FSchematicGraphNode* InNodeToRemove)
 {
+	const FGuid GuidToRemove = InNodeToRemove->GetGuid();
+	NodeByGuid.Remove(GuidToRemove);
+	
 	for (int32 Iter = 0; Iter != Children.Num(); ++Iter)
 	{
 		TSharedRef<SSchematicGraphNode> Child = GetChild(Iter);
-		if (Child->NodeData == InNodeToRemove)
+		if (Child->GetGuid() == GuidToRemove)
 		{
 			Children.RemoveAt(Iter);
 			break;
@@ -340,7 +547,7 @@ void SSchematicGraphPanel::RemoveNode(FSchematicGraphNode* InNodeToRemove)
 	for (int32 Iter = 0; Iter != VisibleChildren.Num(); ++Iter)
 	{
 		TSharedRef<SSchematicGraphNode> Child = StaticCastSharedRef<SSchematicGraphNode>(VisibleChildren[Iter]);
-		if (Child->NodeData == InNodeToRemove)
+		if (Child->GetGuid() == GuidToRemove)
 		{
 			VisibleChildren.RemoveAt(Iter);
 			break;
@@ -348,26 +555,18 @@ void SSchematicGraphPanel::RemoveNode(FSchematicGraphNode* InNodeToRemove)
 	}
 }
 
-void SSchematicGraphPanel::RenameNode(FSchematicGraphNode* InNodeToRename)
+const SSchematicGraphNode* SSchematicGraphPanel::FindNode(const FGuid& InGuid) const
 {
-	for (int32 Iter = 0; Iter != Children.Num(); ++Iter)
+	if(const TSharedPtr<SSchematicGraphNode>* FoundNodePtr = NodeByGuid.Find(InGuid))
 	{
-		TSharedRef<SSchematicGraphNode> Child = GetChild(Iter);
-		if (Child->NodeData == InNodeToRename)
-		{
-			Child->SetToolTipText(FText::FromString(InNodeToRename->Name));
-			break;
-		}
-	}
-	for (int32 Iter = 0; Iter != VisibleChildren.Num(); ++Iter)
-	{
-		TSharedRef<SSchematicGraphNode> Child = StaticCastSharedRef<SSchematicGraphNode>(VisibleChildren[Iter]);
-		if (Child->NodeData == InNodeToRename)
-		{
-			Child->SetToolTipText(FText::FromString(InNodeToRename->Name));
-			break;
-		}
+		return FoundNodePtr->Get();
 	}
+	return nullptr;
+}
+
+void SSchematicGraphPanel::OnArrangeChildren(const FGeometry& AllottedGeometry, FArrangedChildren& ArrangedChildren) const
+{
+	SNodePanel::OnArrangeChildren(AllottedGeometry, ArrangedChildren);
 }
 
 int32 SSchematicGraphPanel::OnPaint(const FPaintArgs& Args, const FGeometry& AllottedGeometry, const FSlateRect& MyCullingRect, FSlateWindowElementList& OutDrawElements, int32 LayerId, const FWidgetStyle& InWidgetStyle, bool bParentEnabled) const
@@ -382,10 +581,11 @@ int32 SSchematicGraphPanel::OnPaint(const FPaintArgs& Args, const FGeometry& All
 		MaxLayerId++;
 	}
 
-	if (FadeBackgroundAlpha->Get() > 0.f)
+	const float BackgroundAlpha = FadeBackgroundAlpha->Get();
+	if (BackgroundAlpha > 0.f)
 	{
 		const FSlateBrush* Brush = FAppStyle::GetBrush(TEXT("Graph.Panel.SolidBackground"));
-		FLinearColor TransparentGrey = FLinearColor(0.5, 0.5, 0.5, FadeBackgroundAlpha->Get());
+		const FLinearColor TransparentGrey = FLinearColor(0.5, 0.5, 0.5, BackgroundAlpha);
 		FSlateDrawElement::MakeBox(
 				OutDrawElements,
 				MaxLayerId,
@@ -420,13 +620,13 @@ int32 SSchematicGraphPanel::OnPaint(const FPaintArgs& Args, const FGeometry& All
 			TSharedRef<SSchematicGraphNode> ChildNode = StaticCastSharedRef<SSchematicGraphNode>(CurWidget.Widget);
 			
 			// Examine node to see what layers we should be drawing in
-			int32 ChildLayerId = NodeLayerId;
+			const int32 ChildLayerId = NodeLayerId;
 
 			const bool bNodeIsVisible = FSlateRect::DoRectanglesIntersect( CurWidget.Geometry.GetLayoutBoundingRect(), MyCullingRect );
 
 			if (bNodeIsVisible)
 			{
-				int32 CurWidgetsMaxLayerId = CurWidget.Widget->Paint(NewArgs, CurWidget.Geometry, MyCullingRect, OutDrawElements, ChildLayerId, InWidgetStyle, true );
+				const int32 CurWidgetsMaxLayerId = CurWidget.Widget->Paint(NewArgs, CurWidget.Geometry, MyCullingRect, OutDrawElements, ChildLayerId, InWidgetStyle, true );
 				MaxLayerId = FMath::Max( MaxLayerId, CurWidgetsMaxLayerId + 1 );
 			}
 		}
@@ -439,6 +639,12 @@ int32 SSchematicGraphPanel::OnPaint(const FPaintArgs& Args, const FGeometry& All
 	return MaxLayerId;
 }
 
+void SSchematicGraphPanel::RemoveAllNodes()
+{
+	NodeByGuid.Reset();
+	SNodePanel::RemoveAllNodes();
+}
+
 TSharedRef<SSchematicGraphNode> SSchematicGraphPanel::GetChild(int32 ChildIndex) const
 {
 	return StaticCastSharedRef<SSchematicGraphNode>(Children[ChildIndex]);
@@ -451,81 +657,77 @@ TStatId SSchematicGraphPanel::GetStatId() const
 
 void SSchematicGraphPanel::Tick(float DeltaTime)
 {
-	bool bIsDragging = false;
+	bIsDragDropping = false;
 	FSlateApplication& Application = FSlateApplication::Get();
 	if (Application.IsDragDropping())
 	{
 		TSharedPtr<FDragDropOperation> DragDropOp = Application.GetDragDroppingContent();
 		if (DragDropOp.IsValid())
 		{
-			bIsDragging = true;
+			bIsDragDropping = true;
 		}
 	}
-	SetFadeBackground(bIsDragging);
 
-	uint16 TopRightNodes = 0;
-	uint16 TopLeftNodes = 0;
-	uint16 BottomRightNodes = 0;
-	uint16 BottomLeftNodes = 0;
+	NodesTopLeft.Reset();
+	NodesTopRight.Reset();
+	NodesBottomLeft.Reset();
+	NodesBottomRight.Reset();
+
 	FVector2D Size = GetCachedGeometry().Size;
 	for (int32 i=0; i<Children.Num(); ++i)
 	{
-		TSharedRef<SSchematicGraphNode> SNode = GetChild(i);
-		UpdateNodeWidgetDelegate.ExecuteIfBound(this, SNode.ToSharedPtr());
+		TSharedRef<SSchematicGraphNode> Node = GetChild(i);
 
-		FVector2d NewPosition = SNode->GetPosition() + (SNode->OriginalSize*0.5);
-		bool bImmediatePosition = !bAnimatePosition;
-		switch(SNode->NodeData->Placement)
+		// update the animation state of the node
+		Node->EnablePositionAnimation(GraphData->GetPositionAnimationEnabledForNode(Node->GetGuid()));
+
+		// collect the nodes displayed in the corners
+		switch(GraphData->GetPlacementForNode(Node->GetNodeData()))
 		{
-			case ESchematicGraphNodePlacement::Panel:
+			case ESchematicGraphNodePlacementConstraint::Free:
 			{
 				break;
 			}
-			case ESchematicGraphNodePlacement::TopLeft:
+			case ESchematicGraphNodePlacementConstraint::TopLeft:
 			{
-				NewPosition = FVector2d(PaddingLeft + SNode->OriginalSize.X*0.5, PaddingTop + SNode->OriginalSize.Y*0.5 + ((SNode->OriginalSize.Y + PaddingInterNode) * TopLeftNodes));
-				TopLeftNodes++;
+				NodesTopLeft.Add(Node->GetGuid());
 				break;
 			}
-			case ESchematicGraphNodePlacement::TopRight:
+			case ESchematicGraphNodePlacementConstraint::TopRight:
 			{
-				NewPosition = FVector2d(Size.X - PaddingRight - (SNode->OriginalSize.X * 0.5), PaddingTop + SNode->OriginalSize.Y*0.5 + ((SNode->OriginalSize.Y + PaddingInterNode) * TopRightNodes));
-				TopRightNodes++;
+				NodesTopRight.Add(Node->GetGuid());
 				break;
 			}
-			case ESchematicGraphNodePlacement::BottomLeft:
+			case ESchematicGraphNodePlacementConstraint::BottomLeft:
 			{
-				NewPosition = FVector2d(PaddingLeft + SNode->OriginalSize.X*0.5, Size.Y - PaddingBottom - (SNode->OriginalSize.Y * 0.5) - ((SNode->OriginalSize.Y + PaddingInterNode)  * BottomLeftNodes));
-				BottomLeftNodes++;
+				NodesBottomLeft.Add(Node->GetGuid());
 				break;
 			}
-			case ESchematicGraphNodePlacement::BottomRight:
+			case ESchematicGraphNodePlacementConstraint::BottomRight:
 			{
-				NewPosition = FVector2d(Size.X - PaddingRight - (SNode->OriginalSize.X * 0.5), Size.Y - PaddingBottom - (SNode->OriginalSize.Y * 0.5) - ((SNode->OriginalSize.Y + PaddingInterNode)  * BottomRightNodes));
-				BottomRightNodes++;
+				NodesBottomRight.Add(Node->GetGuid());
 				break;
 			}
 		}
 
-		if (SNode->bIsBeingDragged)
+		if (Node->IsBeingDragged())
 		{
-			if (bIsDragging)
+			if (bIsDragDropping)
 			{
 				const FVector2f AbsoluteMousePosition = FSlateApplication::Get().GetCursorPos();
 				const FGeometry& Geometry = GetTickSpaceGeometry();
 
 				const FVector2d LocalMousePosition =  (AbsoluteMousePosition - Geometry.GetAbsolutePosition()) / Geometry.GetAccumulatedLayoutTransform().GetScale();
-				NewPosition = LocalMousePosition;
-				bImmediatePosition = true;
-				
+				const FVector2d HalfOriginalSize = Node->GetOriginalSize() * 0.5;
+				Node->PositionDuringDrag = LocalMousePosition + HalfOriginalSize;
 			}
 			else if (DeltaTime > 0.f)
 			{
-				SNode->bIsBeingDragged = false;
+				Node->PositionDuringDrag.Reset();
+				Node->OffsetDuringDrag.Reset();
+				Node->bIsBeingDragged = false;
 			}
 		}
-
-		SNode->SetPosition(NewPosition, bImmediatePosition);
 	}
 }
 
@@ -555,6 +757,112 @@ void SSchematicGraphPanel::OnDropEvent(SSchematicGraphNode* Node, const FDragDro
 	OnDropDelegate.ExecuteIfBound(this, Node, InDragDropEvent);
 }
 
+FVector2d SSchematicGraphPanel::GetPositionForNode(FGuid InNodeGuid) const
+{
+	const SSchematicGraphNode* NodeWidget = FindNode(InNodeGuid);
+	if(GraphData && NodeWidget)
+	{
+		if(NodeWidget->PositionDuringDrag.IsSet())
+		{
+			return NodeWidget->PositionDuringDrag.GetValue() + NodeWidget->OffsetDuringDrag.Get(FVector2d::ZeroVector);
+		}
+		
+		if(const FSchematicGraphNode* Node = GraphData->FindNode(InNodeGuid))
+		{
+			FVector2d Position = GraphData->GetPositionOffsetForNode(Node);
+
+			const FVector2d OriginalSize = NodeWidget->GetOriginalSize();
+			const FVector2D Size = GetCachedGeometry().Size;
+			const FVector2d HalfOriginalSize = OriginalSize * 0.5;
+
+			switch(GraphData->GetPlacementForNode(Node))
+			{
+				case ESchematicGraphNodePlacementConstraint::Free:
+				{
+					Position += GraphData->GetPositionForNode(Node);
+					break;
+				}
+				case ESchematicGraphNodePlacementConstraint::TopLeft:
+				{
+					const int32 NodeIndexInCorner = NodesTopLeft.IndexOfByKey(InNodeGuid);
+					check(NodeIndexInCorner != INDEX_NONE);
+					Position += FVector2d(PaddingLeft + HalfOriginalSize.X, PaddingTop + HalfOriginalSize.Y + ((OriginalSize.Y + PaddingInterNode) * NodeIndexInCorner));
+					break;
+				}
+				case ESchematicGraphNodePlacementConstraint::TopRight:
+				{
+					const int32 NodeIndexInCorner = NodesTopRight.IndexOfByKey(InNodeGuid);
+					check(NodeIndexInCorner != INDEX_NONE);
+					Position += FVector2d(Size.X - PaddingRight - HalfOriginalSize.X, PaddingTop + HalfOriginalSize.Y + ((OriginalSize.Y + PaddingInterNode) * NodeIndexInCorner));
+					break;
+				}
+				case ESchematicGraphNodePlacementConstraint::BottomLeft:
+				{
+					const int32 NodeIndexInCorner = NodesBottomLeft.IndexOfByKey(InNodeGuid);
+					check(NodeIndexInCorner != INDEX_NONE);
+					Position += FVector2d(PaddingLeft + HalfOriginalSize.X, Size.Y - PaddingBottom - HalfOriginalSize.Y - ((OriginalSize.Y + PaddingInterNode)  * NodeIndexInCorner));
+					break;
+				}
+				case ESchematicGraphNodePlacementConstraint::BottomRight:
+				{
+					const int32 NodeIndexInCorner = NodesBottomRight.IndexOfByKey(InNodeGuid);
+					check(NodeIndexInCorner != INDEX_NONE);
+					Position += FVector2d(Size.X - PaddingRight - HalfOriginalSize.X, Size.Y - PaddingBottom - HalfOriginalSize.Y - ((OriginalSize.Y + PaddingInterNode)  * NodeIndexInCorner));
+					break;
+				}
+			}
+			return Position;
+		}
+	}
+	return FVector2d::ZeroVector;
+}
+
+FVector2d SSchematicGraphPanel::GetSizeForNode(FGuid InNodeGuid) const
+{
+	if(GraphData)
+	{
+		return GraphData->GetSizeForNode(InNodeGuid);
+	}
+	static const FVector2d DefaultSize = FVector2d(50.0,50.0);
+	return DefaultSize;
+}
+
+float SSchematicGraphPanel::GetScaleForNode(FGuid InNodeGuid) const
+{
+	if(GraphData)
+	{
+		return GraphData->GetScaleForNode(InNodeGuid);
+	}
+	return 1.f;
+}
+
+FLinearColor SSchematicGraphPanel::GetColorForNode(FGuid InNodeGuid) const
+{
+	if(GraphData)
+	{
+		return GraphData->GetColorForNode(InNodeGuid);
+	}
+	return FLinearColor::White;
+}
+
+const FSlateBrush* SSchematicGraphPanel::GetBrushForNode(FGuid InNodeGuid) const
+{
+	if(GraphData)
+	{
+		return GraphData->GetBrushForNode(InNodeGuid);
+	}
+	static const FSlateBrush* DefaultBrush = SSchematicGraphNode::FArguments()._Brush.Get();
+	return DefaultBrush;
+}
+
+FText SSchematicGraphPanel::GetToolTipForNode(FGuid InNodeGuid) const
+{
+	if(GraphData)
+	{
+		return GraphData->GetToolTipForNode(InNodeGuid);
+	}
+	return FText();
+}
 
 #undef LOCTEXT_NAMESPACE
 
diff --git a/Engine/Source/Developer/AnimationWidgets/Public/SchematicGraphPanel/SSchematicGraphPanel.h b/Engine/Source/Developer/AnimationWidgets/Public/SchematicGraphPanel/SSchematicGraphPanel.h
--- a/Engine/Source/Developer/AnimationWidgets/Public/SchematicGraphPanel/SSchematicGraphPanel.h
+++ b/Engine/Source/Developer/AnimationWidgets/Public/SchematicGraphPanel/SSchematicGraphPanel.h
@@ -1,53 +1,232 @@
 // Copyright Epic Games, Inc. All Rights Reserved.
-#if WITH_EDITOR
 #pragma once
 
+#if WITH_EDITOR
+
 #include "CoreMinimal.h"
 #include "SNodePanel.h"
 #include "TickableEditorObject.h"
 #include "Animation/AnimatedAttribute.h"
 
-enum ESchematicGraphNodePlacement
+enum ESchematicGraphNodePlacementConstraint
 {
-	Panel,
+	Free,
 	TopLeft,
 	TopRight,
 	BottomLeft,
 	BottomRight
 };
 
-struct ANIMATIONWIDGETS_API FSchematicGraphNode
+enum ESchematicGraphNodeVisibility
+{
+	Visible,
+	FadedOut,
+	Hidden
+};
+
+#define SCHEMATICGRAPHNODE_BODY(ClassName, SuperClass) \
+inline static const FName& Type = TEXT(#ClassName); \
+virtual const FName& GetType() const override { return ClassName::Type; } \
+virtual bool IsA(const FName& InType) const \
+{ \
+	if(ClassName::Type == InType) \
+	{ \
+		return true; \
+	} \
+	return SuperClass::IsA(InType); \
+} \
+template<typename T> \
+friend const T* Cast(const ClassName* InNode) \
+{ \
+	return Cast<T>((const FSchematicGraphNode*) InNode); \
+} \
+template<typename T> \
+friend T* Cast(ClassName* InNode) \
+{ \
+	return Cast<T>((FSchematicGraphNode*) InNode); \
+} \
+template<typename T> \
+friend const T* CastChecked(const ClassName* InNode) \
+{ \
+	return CastChecked<T>((const FSchematicGraphNode*) InNode); \
+} \
+template<typename T> \
+friend T* CastChecked(ClassName* InNode) \
+{ \
+	return CastChecked<T>((FSchematicGraphNode*) InNode); \
+} \
+
+class ANIMATIONWIDGETS_API FSchematicGraphNode
 {
-	FString Name;
+public:
+	virtual ~FSchematicGraphNode() {}
+	
+	inline static const FName& Type = TEXT("FSchematicGraphNode");
+	virtual const FName& GetType() const;
+	virtual bool IsA(const FName& InType) const;
+	template<typename NodeType>
+	bool IsA() const
+	{
+		return IsA(NodeType::Type);
+	}
+	
+	template<typename T>
+	friend const T* Cast(const FSchematicGraphNode* InNode)
+	{
+		if(InNode)
+		{
+			if(InNode->IsA<T>())
+			{
+				return static_cast<const T*>(InNode);
+			}
+		}
+		return nullptr;
+	}
+
+	template<typename T>
+	friend T* Cast(FSchematicGraphNode* InNode)
+	{
+		if(InNode)
+		{
+			if(InNode->IsA<T>())
+			{
+				return static_cast<T*>(InNode);
+			}
+		}
+		return nullptr;
+	}
+
+	template<typename T>
+	friend const T* CastChecked(const FSchematicGraphNode* InNode)
+	{
+		const T* Node = Cast<T>(InNode);
+		check(Node);
+		return Node;
+	}
+
+	template<typename T>
+	friend T* CastChecked(FSchematicGraphNode* InNode)
+	{
+		T* Node = Cast<T>(InNode);
+		check(Node);
+		return Node;
+	}
+	
+	const FGuid& GetGuid() const { return Guid; }
+	bool IsSelected() const { return bIsSelected; }
+	void SetSelected(bool bSelected = true) { bIsSelected = bSelected; }
+	virtual FVector2d GetPosition() const { return Position; }
+	virtual FVector2d GetPositionOffset() const { return PositionOffset; }
+	virtual FLinearColor GetColor() const { return Color; }
+	virtual const FSlateBrush* GetBrush() const { return Brush; }
+	virtual const FText& GetToolTip() const { return ToolTip; }
+	virtual ESchematicGraphNodePlacementConstraint GetPlacement() const { return Placement; }
+	virtual void SetPlacement(ESchematicGraphNodePlacementConstraint InPlacement) { Placement = InPlacement; }
+	virtual ESchematicGraphNodeVisibility GetVisibility() const { return Visibility; }
+	virtual void SetVisibility(ESchematicGraphNodeVisibility InVisibility) { Visibility = InVisibility; }
+
+	virtual FString GetDragDropDecoratorLabel() const;
+
+protected:
+	
+	FGuid Guid = FGuid::NewGuid();
 	bool bIsSelected = false;
-	bool bFade = false;
-	ESchematicGraphNodePlacement Placement;
+	FVector2d Position = FVector2d::ZeroVector;
+	FVector2d PositionOffset = FVector2d::ZeroVector;
+	FLinearColor Color = FLinearColor::White;
+	const FSlateBrush* Brush = nullptr;
+	FText ToolTip = FText();
+	ESchematicGraphNodePlacementConstraint Placement = ESchematicGraphNodePlacementConstraint::Free;
+	ESchematicGraphNodeVisibility Visibility = ESchematicGraphNodeVisibility::Visible;
+
+	friend class FSchematicGraphModel;
 };
 
-DECLARE_EVENT_OneParam(FSchematicGraph, FOnNodeAdded, FSchematicGraphNode*);
-DECLARE_EVENT_OneParam(FSchematicGraph, FOnNodeRemoved, FSchematicGraphNode*);
-DECLARE_EVENT_OneParam(FSchematicGraph, FOnNodeRenamed, FSchematicGraphNode*);
+DECLARE_EVENT_OneParam(FSchematicGraph, FOnNodeAdded, const FSchematicGraphNode*);
+DECLARE_EVENT_OneParam(FSchematicGraph, FOnNodeRemoved, const FSchematicGraphNode*);
 DECLARE_EVENT(FSchematicGraph, FOnGraphReset);
 
-class ANIMATIONWIDGETS_API FSchematicGraph
+class ANIMATIONWIDGETS_API FSchematicGraphModel
 {
 public:
-	TArray<FSchematicGraphNode*> Nodes;
-	TMap<FSchematicGraphNode*, TArray<FSchematicGraphNode*>> Links;
 
-	void Reset();
-	FSchematicGraphNode* AddNode(const FString& InName);
-	bool RenameNode(const FString& InOldName, const FString& InNewName);
-	bool RemoveNode(const FString& InName);
+	virtual ~FSchematicGraphModel() {}
+	
+	virtual void Reset();
+	
+	template<typename NodeType = FSchematicGraphNode>
+	NodeType* AddNode(bool bNotify = true)
+	{
+		const TSharedPtr<FSchematicGraphNode> NewNode = MakeShareable(new NodeType);
+		Nodes.Add(NewNode);
+		NodeByGuid.Add(NewNode->GetGuid(), NewNode);
+
+		if (bNotify && OnNodeAddedDelegate.IsBound())
+		{
+			OnNodeAddedDelegate.Broadcast(NewNode.Get());
+		}
+		return static_cast<NodeType*>(NewNode.Get());
+	}
+
+	template<typename NodeType = FSchematicGraphNode>
+	const NodeType* FindNode(const FGuid& InGuid) const
+	{
+		if(const TSharedPtr<FSchematicGraphNode>* ExistingNode = NodeByGuid.Find(InGuid))
+		{
+			return Cast<NodeType>(ExistingNode->Get());
+		};
+		return nullptr;
+	}
+
+	template<typename NodeType = FSchematicGraphNode>
+	const NodeType* FindNodeChecked(const FGuid& InGuid) const
+	{
+		if(const TSharedPtr<FSchematicGraphNode>* ExistingNode = NodeByGuid.Find(InGuid))
+		{
+			check(ExistingNode->Get()->IsA(NodeType::Type));
+			return static_cast<NodeType*>(ExistingNode->Get());
+		};
+		return nullptr;
+	}
+
+	virtual bool RemoveNode(const FGuid& InGuid);
+
+	const TArray<TSharedPtr<FSchematicGraphNode>>& GetNodes() const { return Nodes; }
+
+	FVector2d GetPositionForNode(const FGuid& InGuid) const;
+	virtual FVector2d GetPositionForNode(const FSchematicGraphNode* InNode) const;
+	FVector2d GetPositionOffsetForNode(const FGuid& InGuid) const;
+	virtual FVector2d GetPositionOffsetForNode(const FSchematicGraphNode* InNode) const;
+	bool GetPositionAnimationEnabledForNode(const FGuid& InNodeGuid) const;
+	virtual bool GetPositionAnimationEnabledForNode(const FSchematicGraphNode* InNode) const;
+	FVector2d GetSizeForNode(const FGuid& InNodeGuid) const;
+	virtual FVector2d GetSizeForNode(const FSchematicGraphNode* InNode) const;
+	float GetScaleForNode(const FGuid& InNodeGuid) const;
+	virtual float GetScaleForNode(const FSchematicGraphNode* InNode) const;
+	FLinearColor GetColorForNode(const FGuid& InGuid) const;
+	virtual FLinearColor GetColorForNode(const FSchematicGraphNode* InNode) const;
+	const FSlateBrush* GetBrushForNode(const FGuid& InGuid) const;
+	virtual const FSlateBrush* GetBrushForNode(const FSchematicGraphNode* InNode) const;
+	const FText GetToolTipForNode(const FGuid& InGuid) const;
+	virtual const FText GetToolTipForNode(const FSchematicGraphNode* InNode) const;
+	ESchematicGraphNodePlacementConstraint GetPlacementForNode(const FGuid& InGuid) const;
+	virtual ESchematicGraphNodePlacementConstraint GetPlacementForNode(const FSchematicGraphNode* InNode) const;
+	ESchematicGraphNodeVisibility GetVisibilityForNode(const FGuid& InGuid) const;
+	virtual ESchematicGraphNodeVisibility GetVisibilityForNode(const FSchematicGraphNode* InNode) const;
 
-	FOnNodeAdded OnNodeAddedDelegate;
 	FOnNodeAdded& OnNodeAdded() { return OnNodeAddedDelegate; }
-	FOnNodeRemoved OnNodeRemovedDelegate;
 	FOnNodeRemoved& OnNodeRemoved() { return OnNodeRemovedDelegate; }
-	FOnNodeRenamed OnNodeRenamedDelegate;
-	FOnNodeRenamed& OnNodeRenamed() { return OnNodeRenamedDelegate; }
-	FOnGraphReset OnGraphResetDelegate;
 	FOnGraphReset& OnGraphReset() { return OnGraphResetDelegate; }
+
+protected:
+
+	TArray<TSharedPtr<FSchematicGraphNode>> Nodes;
+	TMap<FGuid, TSharedPtr<FSchematicGraphNode>> NodeByGuid;
+	TMap<FGuid, TArray<FGuid>> Links;
+
+	FOnNodeAdded OnNodeAddedDelegate;
+	FOnNodeRemoved OnNodeRemovedDelegate;
+	FOnGraphReset OnGraphResetDelegate;
 };
 
 class SSchematicGraphNode;
@@ -59,18 +238,9 @@ public:
 
 	DECLARE_DELEGATE_TwoParams(FOnEndDrag, SSchematicGraphNode*, const FDragDropOperation&);
 
-	static TSharedRef<FSchematicGraphNodeDragDropOp> New(TArray<SSchematicGraphNode*> InSchematicGraphNodes, const TArray<FString>& InElements, FOnEndDrag InOnEndDragDelegate);
+	static TSharedRef<FSchematicGraphNodeDragDropOp> New(TArray<SSchematicGraphNode*> InSchematicGraphNodes, const TArray<FGuid>& InElements, FOnEndDrag InOnEndDragDelegate);
 
-	~FSchematicGraphNodeDragDropOp()
-	{
-		if (OnEndDragDelegate.IsBound())
-		{
-			for (SSchematicGraphNode* Node : SchematicGraphNodes)
-			{
-				OnEndDragDelegate.Execute(Node, *this);
-			}
-		}
-	}
+	virtual ~FSchematicGraphNodeDragDropOp() override;
 	
 	virtual TSharedPtr<SWidget> GetDefaultDecorator() const override;
 
@@ -81,12 +251,12 @@ public:
 	}
 
 	/** @return The property paths from this drag operation */
-	const TArray<FString>& GetElements() const
+	const TArray<FGuid>& GetElements() const
 	{
 		return Elements;
 	}
 
-	FString GetJoinedElementNames() const;
+	FString GetJoinedDecoratorLabels() const;
 
 private:
 
@@ -94,7 +264,7 @@ private:
 	TArray<SSchematicGraphNode*> SchematicGraphNodes;
 
 	/** Data for the property paths this item represents */
-	TArray<FString> Elements;
+	TArray<FGuid> Elements;
 
 	/** Delegate to call when this drag operation ends */
 	FOnEndDrag OnEndDragDelegate;
@@ -103,17 +273,28 @@ private:
 class ANIMATIONWIDGETS_API SSchematicGraphNode : public SNodePanel::SNode
 {
 public:
+
+	typedef TAnimatedAttribute<FVector2d> FVector2dAttribute;
+	typedef TAnimatedAttribute<float> FFloatAttribute;
+	typedef TAnimatedAttribute<FLinearColor> FLinearColorAttribute;
+	
 	DECLARE_DELEGATE_OneParam(FOnClicked, SSchematicGraphNode*);
 	DECLARE_DELEGATE_TwoParams(FOnBeginDrag, SSchematicGraphNode*, const FDragDropOperation&);
 	DECLARE_DELEGATE_TwoParams(FOnEndDrag, SSchematicGraphNode*, const FDragDropOperation&);
 	DECLARE_DELEGATE_TwoParams(FOnDrop, SSchematicGraphNode*, const FDragDropEvent&);
 
-	SLATE_BEGIN_ARGS(SSchematicGraphNode){}
-	SLATE_ARGUMENT(FSchematicGraphNode*, NodeData)
+	SLATE_BEGIN_ARGS(SSchematicGraphNode);
+	SLATE_ARGUMENT(const FSchematicGraphNode*, NodeData)
+	SLATE_ARGUMENT(TSharedPtr<FVector2dAttribute>, Position)
+	SLATE_ARGUMENT(TSharedPtr<FVector2dAttribute>, Size)
+	SLATE_ARGUMENT(TSharedPtr<FFloatAttribute>, Scale)
+	SLATE_ARGUMENT(TSharedPtr<FLinearColorAttribute>, Color)
+	SLATE_ATTRIBUTE(const FSlateBrush*, Brush)
 	SLATE_EVENT(FOnClicked, OnClicked)
 	SLATE_EVENT(FOnBeginDrag, OnBeginDrag)
 	SLATE_EVENT(FOnEndDrag, OnEndDrag)
 	SLATE_EVENT(FOnDrop, OnDrop)
+	SLATE_ATTRIBUTE(FText, ToolTipText)
 	SLATE_END_ARGS()
 
 	void Construct(const FArguments& InArgs);
@@ -130,23 +311,37 @@ public:
 	virtual FReply OnDragDetected(const FGeometry& MyGeometry, const FPointerEvent& MouseEvent) override;
 
 	virtual FVector2d GetPosition() const override;
+	void EnablePositionAnimation(bool bEnabled = true);
+
+	const FVector2d& GetOriginalSize() const { return OriginalSize; }
 
-	void SetPosition(const FVector2d& InPosition, bool bImmediately = false);
+	const FSchematicGraphNode* GetNodeData() const { return NodeData; }
+	const FGuid GetGuid() const;
+	bool IsInteractive() const;
+
+	const bool IsBeingDragged() const { return bIsBeingDragged; } 
+	
+private:
 
 	bool bIsBeingDragged = false;
 	FVector2d OriginalSize = FVector2d(50.0,50.0);
 	float ScaledUp = 1.25;
 	float ScaledDown = 0.75;
 
+	const FSchematicGraphNode* NodeData = nullptr;
+	TSharedPtr<FVector2dAttribute> Position;
+	TOptional<FVector2d> PositionDuringDrag;
+	TOptional<FVector2d> OffsetDuringDrag;
+	TSharedPtr<FVector2dAttribute> Size;
+	TSharedPtr<FFloatAttribute> Scale;
+	TSharedPtr<FLinearColorAttribute> Color;
+	TAttribute<const FSlateBrush*> Brush;
 	FOnClicked OnClickedDelegate;
 	FOnBeginDrag OnBeginDragDelegate;
 	FOnEndDrag OnEndDragDelegate;
 	FOnDrop OnDropDelegate;
-	FSchematicGraphNode* NodeData;
-	TSharedPtr<TAnimatedAttribute<FVector2d>> Position;
-	TSharedPtr<TAnimatedAttribute<FVector2d>> Size;
-	TSharedPtr<TAnimatedAttribute<float>> Scale;
-	FSlateBrush Brush;
+
+	friend class SSchematicGraphPanel;
 };
 
 
@@ -154,8 +349,11 @@ public:
 class ANIMATIONWIDGETS_API SSchematicGraphPanel : public SNodePanel, public FTickableEditorObject
 {
 public:
-	
-	DECLARE_DELEGATE_TwoParams(FUpdateNodeWidget, SSchematicGraphPanel*, TSharedPtr<SSchematicGraphNode>);
+
+	using FVector2dAttribute = SSchematicGraphNode::FVector2dAttribute;
+	using FFloatAttribute = SSchematicGraphNode::FFloatAttribute;
+	using FLinearColorAttribute = SSchematicGraphNode::FLinearColorAttribute;
+
 	DECLARE_DELEGATE_TwoParams(FOnNodeClicked, SSchematicGraphPanel*, SSchematicGraphNode*);
 	DECLARE_DELEGATE_ThreeParams(FOnBeginDrag, SSchematicGraphPanel*, SSchematicGraphNode*, const FDragDropOperation&);
 	DECLARE_DELEGATE_ThreeParams(FOnEndDrag, SSchematicGraphPanel*, SSchematicGraphNode*, const FDragDropOperation&);
@@ -163,36 +361,36 @@ public:
 	
 	SLATE_BEGIN_ARGS(SSchematicGraphPanel) {}
 	SLATE_ARGUMENT(bool, IsOverlay)
-	SLATE_ARGUMENT(FSchematicGraph*, GraphData)
+	SLATE_ARGUMENT(FSchematicGraphModel*, GraphData)
 	SLATE_ARGUMENT(int32, PaddingLeft)
 	SLATE_ARGUMENT(int32, PaddingRight)
 	SLATE_ARGUMENT(int32, PaddingTop)
 	SLATE_ARGUMENT(int32, PaddingBottom)
 	SLATE_ARGUMENT(int32, PaddingInterNode)
-	SLATE_EVENT(FUpdateNodeWidget, OnUpdateNodeWidget)
 	SLATE_EVENT(FOnNodeClicked, OnNodeClicked)
 	SLATE_EVENT(FOnBeginDrag, OnBeginDrag)
 	SLATE_EVENT(FOnEndDrag, OnEndDrag)
 	SLATE_EVENT(FOnDrop, OnDrop)
 	SLATE_END_ARGS()
 
-	~SSchematicGraphPanel()
+	virtual ~SSchematicGraphPanel() override
 	{
-		UpdateNodeWidgetDelegate.Unbind();
 		OnNodeClickedDelegate.Unbind();
 		OnDropDelegate.Unbind();
 	}
 
-	void SetSchematicGraph(FSchematicGraph* InGraphData);
+	void SetSchematicGraph(FSchematicGraphModel* InGraphData);
 	void Construct(const FArguments& InArgs);
 
 	void RebuildPanel();
-	void AddNode(FSchematicGraphNode* NodeToAdd);
-	void RemoveNode(FSchematicGraphNode* NodeToRemove);
-	void RenameNode(FSchematicGraphNode* NodeToRemove);
-
+	void AddNode(const FSchematicGraphNode* InNodeToAdd);
+	void RemoveNode(const FSchematicGraphNode* InNodeToRemove);
+	const SSchematicGraphNode* FindNode(const FGuid& InGuid) const;
+	
 	// SNodePanel interface
+	virtual void OnArrangeChildren(const FGeometry& AllottedGeometry, FArrangedChildren& ArrangedChildren) const override;
 	virtual int32 OnPaint(const FPaintArgs& Args, const FGeometry& AllottedGeometry, const FSlateRect& MyCullingRect, FSlateWindowElementList& OutDrawElements, int32 LayerId, const FWidgetStyle& InWidgetStyle, bool bParentEnabled) const override;
+	virtual void RemoveAllNodes() override;
 	// End of SNodePanel interface
 	
 	TSharedRef<SSchematicGraphNode> GetChild(int32 ChildIndex) const;
@@ -209,22 +407,31 @@ public:
 	void OnEndDragEvent(SSchematicGraphNode* Node, const FDragDropOperation& InDragDropEvent);
 	void OnDropEvent(SSchematicGraphNode* Node, const FDragDropEvent& InDragDropEvent);
 
-	void SetFadeBackground(bool bInFade) { FadeBackgroundAlpha->Set(bInFade ? 0.5f : 0.f); }
+	virtual FVector2d GetPositionForNode(FGuid InNodeGuid) const;
+	virtual FVector2d GetSizeForNode(FGuid InNodeGuid) const;
+	virtual float GetScaleForNode(FGuid InNodeGuid) const;
+	virtual FLinearColor GetColorForNode(FGuid InNodeGuid) const;
+	virtual const FSlateBrush* GetBrushForNode(FGuid InNodeGuid) const;
+	virtual FText GetToolTipForNode(FGuid InNodeGuid) const;
 
-	bool bAnimatePosition = true;
+	bool bIsDragDropping = false;
 	bool bIsOverlay;
 	int32 PaddingLeft = 0;
 	int32 PaddingRight = 0;
 	int32 PaddingTop = 0;
 	int32 PaddingBottom = 0;
 	int32 PaddingInterNode = 0;
-	TSharedPtr<TAnimatedAttribute<float>> FadeBackgroundAlpha;
-	FSchematicGraph* GraphData;
-	FUpdateNodeWidget UpdateNodeWidgetDelegate;
+	TSharedPtr<FFloatAttribute> FadeBackgroundAlpha;
+	FSchematicGraphModel* GraphData;
 	FOnNodeClicked OnNodeClickedDelegate;
 	FOnBeginDrag OnBeginDragDelegate;
 	FOnEndDrag OnEndDragDelegate;
 	FOnDrop OnDropDelegate;
+	TArray<FGuid> NodesTopLeft;
+	TArray<FGuid> NodesTopRight;
+	TArray<FGuid> NodesBottomLeft;
+	TArray<FGuid> NodesBottomRight;
+	TMap<FGuid, TSharedPtr<SSchematicGraphNode>> NodeByGuid;
 };
 
 #endif
\ No newline at end of file
diff --git a/Engine/Source/Developer/AnimationWidgets/Public/SchematicGraphPanel/SchematicViewportTabSummoner.h b/Engine/Source/Developer/AnimationWidgets/Public/SchematicGraphPanel/SchematicViewportTabSummoner.h
--- a/Engine/Source/Developer/AnimationWidgets/Public/SchematicGraphPanel/SchematicViewportTabSummoner.h
+++ b/Engine/Source/Developer/AnimationWidgets/Public/SchematicGraphPanel/SchematicViewportTabSummoner.h
@@ -17,7 +17,7 @@ struct ANIMATIONWIDGETS_API FSchematicViewportArgs
 	{}
 
 	/** The model which contains the graph to display */
-	FSchematicGraph* SchematicGraph;
+	FSchematicGraphModel* SchematicGraph;
 	
 	/** Delegate fired when the viewport is created */
 	FOnSchematicViewportCreated OnViewportCreated;
@@ -35,7 +35,7 @@ public:
 	virtual FTabSpawnerEntry& RegisterTabSpawner(TSharedRef<FTabManager> TabManager, const FApplicationMode* CurrentApplicationMode) const;
 	virtual TSharedRef<SWidget> CreateTabBody(const FWorkflowTabSpawnInfo& Info) const override;
 
-	FSchematicGraph* SchematicGraph;
+	FSchematicGraphModel* SchematicGraph;
 	FOnSchematicViewportCreated OnViewportCreated;
 };
 #endif
\ No newline at end of file
diff --git a/Engine/Source/Runtime/SlateCore/Public/Animation/AnimatedAttribute.h b/Engine/Source/Runtime/SlateCore/Public/Animation/AnimatedAttribute.h
--- a/Engine/Source/Runtime/SlateCore/Public/Animation/AnimatedAttribute.h
+++ b/Engine/Source/Runtime/SlateCore/Public/Animation/AnimatedAttribute.h
@@ -68,12 +68,13 @@ public:
 	 *
 	 * @param  InSettings		The settings for the interpolator
 	 * @param  InGetter			Delegate to bind
+	 * @param  InDefaultValue   The optional default value. If not specified the getter will be used to determine the default
 	 */
 	template< typename InterpolatorSettings >
-	[[nodiscard]] static TSharedRef< TAnimatedAttribute > CreateWithGetter( const InterpolatorSettings& InSettings, const FGetter& InGetter )
+	[[nodiscard]] static TSharedRef< TAnimatedAttribute > CreateWithGetter( const InterpolatorSettings& InSettings, const FGetter& InGetter, const TOptional<NumericType>& InDefaultValue = TOptional<NumericType>() )
 	{
 		TSharedRef< TAnimatedAttribute< NumericType > > Attribute = MakeShared< TAnimatedAttribute< NumericType > >(FPrivateToken{}, InSettings);
-		Attribute->Set(InGetter.Execute());
+		Attribute->Set(InDefaultValue.IsSet() ? InDefaultValue.GetValue() : InGetter.Execute());
 		Attribute->Getter = InGetter;
 		Attribute->Register();
 		return Attribute;
@@ -84,12 +85,13 @@ public:
 	 *
 	 * @param  InSettings		The settings for the interpolator
 	 * @param  InGetter			Delegate to bind
+	 * @param  InDefaultValue   The optional default value. If not specified the getter will be used to determine the default
 	 */
 	template< typename InterpolatorSettings >
-	[[nodiscard]] static TSharedRef< TAnimatedAttribute > CreateWithGetter( const InterpolatorSettings& InSettings, FGetter&& InGetter )
+	[[nodiscard]] static TSharedRef< TAnimatedAttribute > CreateWithGetter( const InterpolatorSettings& InSettings, FGetter&& InGetter, const TOptional<NumericType>& InDefaultValue = TOptional<NumericType>() )
 	{
 		TSharedRef< TAnimatedAttribute< NumericType > > Attribute = MakeShared< TAnimatedAttribute< NumericType > >(FPrivateToken{}, InSettings);
-		Attribute->Set(InGetter.Execute());
+		Attribute->Set(InDefaultValue.IsSet() ? InDefaultValue.GetValue() : InGetter.Execute());
 		Attribute->Getter = MoveTemp(InGetter);
 		Attribute->Register();
 		return Attribute;
