Allow programs to build with AnimationWidgets

diff --git a/Engine/Source/Developer/AnimationWidgets/AnimationWidgets.Build.cs b/Engine/Source/Developer/AnimationWidgets/AnimationWidgets.Build.cs
--- a/Engine/Source/Developer/AnimationWidgets/AnimationWidgets.Build.cs
+++ b/Engine/Source/Developer/AnimationWidgets/AnimationWidgets.Build.cs
@@ -20,7 +20,7 @@ public class AnimationWidgets : ModuleRules
 			}
 		);
 
-		if (Target.Type == TargetType.Editor)
+		if (Target.Type == TargetType.Editor || Target.Type == TargetType.Program)
 		{
 			PrivateDependencyModuleNames.AddRange(
 				new string[]
